- Movied previously committed squencer changes for music mode into a new Musical Mode plugin
- Changes retiming of track keyframes to support new clock plugin system
- Changes to UI system to support new clock plugin system
- Changes to numeric interfaces
- Removed music map from Movie Sequence in keeping with new clock plugin system.
- Removed experimentral editor setting for turning on music mode in sequencer now that music mode is in its own experimental plugin.
- New Music Environment Engine Subsystem
- Changes to the Harmonix plugin to support the Music Environment Subsystem

Reviews in multiple changelists to simplify review process for massive number of files...

https://p4-swarm.epicgames.net/reviews/38388494
https://p4-swarm.epicgames.net/reviews/38389457
https://p4-swarm.epicgames.net/reviews/38389466
https://p4-swarm.epicgames.net/reviews/38389479
https://p4-swarm.epicgames.net/reviews/38389491
https://p4-swarm.epicgames.net/reviews/38389513

diff --git a/Engine/Plugins/Animation/AnimationData/Source/AnimationData/Private/AnimSequencerController.cpp b/Engine/Plugins/Animation/AnimationData/Source/AnimationData/Private/AnimSequencerController.cpp
--- a/Engine/Plugins/Animation/AnimationData/Source/AnimationData/Private/AnimSequencerController.cpp
+++ b/Engine/Plugins/Animation/AnimationData/Source/AnimationData/Private/AnimSequencerController.cpp
@@ -3378,7 +3378,7 @@ bool UAnimSequencerController::SetCurveControlKey(const FName& CurveName, const
 					if(bOffFrameCurveKey)
 					{
 						const FFrameRate MaxChannelRate = FFrameRate(240000, 1);
-						ParameterCurvePair->ParameterCurve.ChangeFrameResolution(CurveRate, MaxChannelRate);
+						ParameterCurvePair->ParameterCurve.RemapTimes(UE::MovieScene::FFrameRateRetiming{ CurveRate, MaxChannelRate });
 						ParameterCurvePair->ParameterCurve.SetTickResolution(MaxChannelRate);
 					}					
 
diff --git a/Engine/Plugins/Animation/ControlRig/Source/ControlRig/Private/Sequencer/MovieSceneControlRigSpaceChannel.cpp b/Engine/Plugins/Animation/ControlRig/Source/ControlRig/Private/Sequencer/MovieSceneControlRigSpaceChannel.cpp
--- a/Engine/Plugins/Animation/ControlRig/Source/ControlRig/Private/Sequencer/MovieSceneControlRigSpaceChannel.cpp
+++ b/Engine/Plugins/Animation/ControlRig/Source/ControlRig/Private/Sequencer/MovieSceneControlRigSpaceChannel.cpp
@@ -82,14 +82,9 @@ void FMovieSceneControlRigSpaceChannel::DeleteKeysFrom(FFrameNumber InTime, bool
 	BroadcastSpaceNoLongerUsed(BeforeKeys, AfterKeys);
 }
 
-void FMovieSceneControlRigSpaceChannel::ChangeFrameResolution(FFrameRate SourceRate, FFrameRate DestinationRate)
+void FMovieSceneControlRigSpaceChannel::RemapTimes(const UE::MovieScene::IRetimingInterface& Retimer)
 {
-	GetData().ChangeFrameResolution(SourceRate, DestinationRate);
-}
-
-void FMovieSceneControlRigSpaceChannel::RemapMusicalFrameTimes(const UFrameBasedMusicMap* OldMap, const UFrameBasedMusicMap* NewMap)
-{
-	GetData().RemapMusicalFrameTimes(OldMap, NewMap);
+	GetData().RemapTimes(Retimer);
 }
 
 TRange<FFrameNumber> FMovieSceneControlRigSpaceChannel::ComputeEffectiveRange() const
diff --git a/Engine/Plugins/Animation/ControlRig/Source/ControlRig/Public/Sequencer/MovieSceneControlRigSpaceChannel.h b/Engine/Plugins/Animation/ControlRig/Source/ControlRig/Public/Sequencer/MovieSceneControlRigSpaceChannel.h
--- a/Engine/Plugins/Animation/ControlRig/Source/ControlRig/Public/Sequencer/MovieSceneControlRigSpaceChannel.h
+++ b/Engine/Plugins/Animation/ControlRig/Source/ControlRig/Public/Sequencer/MovieSceneControlRigSpaceChannel.h
@@ -106,8 +106,7 @@ public:
 	virtual void DuplicateKeys(TArrayView<const FKeyHandle> InHandles, TArrayView<FKeyHandle> OutNewHandles) override;
 	virtual void DeleteKeys(TArrayView<const FKeyHandle> InHandles) override;
 	virtual void DeleteKeysFrom(FFrameNumber InTime, bool bDeleteKeysBefore) override;
-	virtual void ChangeFrameResolution(FFrameRate SourceRate, FFrameRate DestinationRate) override;
-	virtual void RemapMusicalFrameTimes(const UFrameBasedMusicMap* OldMap, const UFrameBasedMusicMap* NewMap) override;
+	virtual void RemapTimes(const UE::MovieScene::IRetimingInterface& Retimer) override;
 	virtual TRange<FFrameNumber> ComputeEffectiveRange() const override;
 	virtual int32 GetNumKeys() const override;
 	virtual void Reset() override;
diff --git a/Engine/Plugins/Animation/ControlRig/Source/ControlRigEditor/Private/EditMode/SControlRigSnapper.cpp b/Engine/Plugins/Animation/ControlRig/Source/ControlRigEditor/Private/EditMode/SControlRigSnapper.cpp
--- a/Engine/Plugins/Animation/ControlRig/Source/ControlRigEditor/Private/EditMode/SControlRigSnapper.cpp
+++ b/Engine/Plugins/Animation/ControlRig/Source/ControlRigEditor/Private/EditMode/SControlRigSnapper.cpp
@@ -290,7 +290,7 @@ void SControlRigSnapper::Construct(const FArguments& InArgs)
 										TSharedPtr<ISequencer> SequencerPtr = Snapper.GetSequencer().Pin();
 										if (SequencerPtr.IsValid() && SequencerPtr->GetFocusedMovieSceneSequence())
 										{
-											TOptional<double> NewFrameTime = SequencerPtr->GetNumericPositionTypeInterface()->FromString(InText.ToString(), 0);
+											TOptional<double> NewFrameTime = SequencerPtr->GetNumericTypeInterface()->FromString(InText.ToString(), 0);
 											if (NewFrameTime.IsSet())
 											{
 												StartFrame = FFrameNumber((int32)NewFrameTime.GetValue());
@@ -324,7 +324,7 @@ void SControlRigSnapper::Construct(const FArguments& InArgs)
 										if (Sequencer.IsValid() && Sequencer.Pin()->GetFocusedMovieSceneSequence())
 										{
 											TSharedPtr<ISequencer> SequencerPtr = Sequencer.Pin();
-											TOptional<double> NewFrameTime = SequencerPtr->GetNumericPositionTypeInterface()->FromString(InText.ToString(), 0);
+											TOptional<double> NewFrameTime = SequencerPtr->GetNumericTypeInterface()->FromString(InText.ToString(), 0);
 											if (NewFrameTime.IsSet())
 											{
 												EndFrame = FFrameNumber((int32)NewFrameTime.GetValue());
@@ -466,7 +466,7 @@ FText SControlRigSnapper::GetStartFrameToSnapText()
 	TWeakPtr<ISequencer> Sequencer = Snapper.GetSequencer();
 	if (Sequencer.IsValid() && Sequencer.Pin()->GetFocusedMovieSceneSequence())
 	{
-		Value = FText::FromString(Sequencer.Pin()->GetNumericPositionTypeInterface()->ToString(StartFrame.Value));
+		Value = FText::FromString(Sequencer.Pin()->GetNumericTypeInterface()->ToString(StartFrame.Value));
 	}
 	return Value;
 }
@@ -477,7 +477,7 @@ FText SControlRigSnapper::GetEndFrameToSnapText()
 	TWeakPtr<ISequencer> Sequencer = Snapper.GetSequencer();
 	if (Sequencer.IsValid() && Sequencer.Pin()->GetFocusedMovieSceneSequence())
 	{
-		Value = FText::FromString(Sequencer.Pin()->GetNumericPositionTypeInterface()->ToString(EndFrame.Value));
+		Value = FText::FromString(Sequencer.Pin()->GetNumericTypeInterface()->ToString(EndFrame.Value));
 	}
 	return Value;
 }
diff --git a/Engine/Plugins/Animation/ControlRig/Source/ControlRigEditor/Private/Editor/Constraints/SConstraintsWidget.cpp b/Engine/Plugins/Animation/ControlRig/Source/ControlRigEditor/Private/Editor/Constraints/SConstraintsWidget.cpp
--- a/Engine/Plugins/Animation/ControlRig/Source/ControlRigEditor/Private/Editor/Constraints/SConstraintsWidget.cpp
+++ b/Engine/Plugins/Animation/ControlRig/Source/ControlRigEditor/Private/Editor/Constraints/SConstraintsWidget.cpp
@@ -1578,9 +1578,8 @@ void SConstraintBakeWidget::Construct(const FArguments& InArgs)
 	FPropertyEditorModule& PropertyEditor = FModuleManager::Get().LoadModuleChecked<FPropertyEditorModule>(TEXT("PropertyEditor"));
 
 	DetailsView = PropertyEditor.CreateStructureDetailView(ViewArgs, StructureViewArgs, TSharedPtr<FStructOnScope>());
-	TSharedPtr<INumericTypeInterface<double>> NumericPositionTypeInterface = Sequencer->GetNumericPositionTypeInterface();
 	DetailsView->GetDetailsView()->RegisterInstancedCustomPropertyTypeLayout("FrameNumber",
-		FOnGetPropertyTypeCustomizationInstance::CreateLambda([=]() {return MakeShared<FFrameNumberDetailsCustomization>(NumericPositionTypeInterface); }));
+		FOnGetPropertyTypeCustomizationInstance::CreateSP(Sequencer.ToSharedRef(), &ISequencer::MakeFrameNumberDetailsCustomization));
 	DetailsView->SetStructureData(Settings);
 
 
diff --git a/Engine/Plugins/Animation/ControlRig/Source/ControlRigEditor/Private/Editor/SRigSpacePickerWidget.cpp b/Engine/Plugins/Animation/ControlRig/Source/ControlRigEditor/Private/Editor/SRigSpacePickerWidget.cpp
--- a/Engine/Plugins/Animation/ControlRig/Source/ControlRigEditor/Private/Editor/SRigSpacePickerWidget.cpp
+++ b/Engine/Plugins/Animation/ControlRig/Source/ControlRigEditor/Private/Editor/SRigSpacePickerWidget.cpp
@@ -1185,9 +1185,9 @@ void SRigSpacePickerBakeWidget::Construct(const FArguments& InArgs)
 	FPropertyEditorModule& PropertyEditor = FModuleManager::Get().LoadModuleChecked<FPropertyEditorModule>(TEXT("PropertyEditor"));
 
 	DetailsView = PropertyEditor.CreateStructureDetailView(ViewArgs, StructureViewArgs, TSharedPtr<FStructOnScope>());
-	TSharedPtr<INumericTypeInterface<double>> NumericPositionTypeInterface = Sequencer->GetNumericPositionTypeInterface();
+
 	DetailsView->GetDetailsView()->RegisterInstancedCustomPropertyTypeLayout("FrameNumber",
-		FOnGetPropertyTypeCustomizationInstance::CreateLambda([=]() {return MakeShared<FFrameNumberDetailsCustomization>(NumericPositionTypeInterface); }));
+		FOnGetPropertyTypeCustomizationInstance::CreateSP(Sequencer, &ISequencer::MakeFrameNumberDetailsCustomization));
 	DetailsView->SetStructureData(Settings);
 
 	ChildSlot
diff --git a/Engine/Plugins/Animation/ControlRig/Source/ControlRigEditor/Private/Sequencer/ControlRigParameterTrackEditor.cpp b/Engine/Plugins/Animation/ControlRig/Source/ControlRigEditor/Private/Sequencer/ControlRigParameterTrackEditor.cpp
--- a/Engine/Plugins/Animation/ControlRig/Source/ControlRigEditor/Private/Sequencer/ControlRigParameterTrackEditor.cpp
+++ b/Engine/Plugins/Animation/ControlRig/Source/ControlRigEditor/Private/Sequencer/ControlRigParameterTrackEditor.cpp
@@ -5066,9 +5066,8 @@ void SCollapseControlsWidget::Construct(const FArguments& InArgs)
 	FPropertyEditorModule& PropertyEditor = FModuleManager::Get().LoadModuleChecked<FPropertyEditorModule>(TEXT("PropertyEditor"));
 
 	DetailsView = PropertyEditor.CreateStructureDetailView(ViewArgs, StructureViewArgs, TSharedPtr<FStructOnScope>());
-	TSharedPtr<INumericTypeInterface<double>> NumericPositionTypeInterface = Sequencer.Pin()->GetNumericPositionTypeInterface();
 	DetailsView->GetDetailsView()->RegisterInstancedCustomPropertyTypeLayout("FrameNumber",
-		FOnGetPropertyTypeCustomizationInstance::CreateLambda([=]() {return MakeShared<FFrameNumberDetailsCustomization>(NumericPositionTypeInterface); }));
+		FOnGetPropertyTypeCustomizationInstance::CreateSP(Sequencer.Pin().ToSharedRef(), &ISequencer::MakeFrameNumberDetailsCustomization));
 	DetailsView->SetStructureData(Settings);
 
 	ChildSlot
diff --git a/Engine/Plugins/FX/Niagara/Source/NiagaraEditor/Private/Sequencer/NiagaraSequence/Sections/MovieSceneNiagaraEmitterSection.cpp b/Engine/Plugins/FX/Niagara/Source/NiagaraEditor/Private/Sequencer/NiagaraSequence/Sections/MovieSceneNiagaraEmitterSection.cpp
--- a/Engine/Plugins/FX/Niagara/Source/NiagaraEditor/Private/Sequencer/NiagaraSequence/Sections/MovieSceneNiagaraEmitterSection.cpp
+++ b/Engine/Plugins/FX/Niagara/Source/NiagaraEditor/Private/Sequencer/NiagaraSequence/Sections/MovieSceneNiagaraEmitterSection.cpp
@@ -85,14 +85,9 @@ void FMovieSceneNiagaraEmitterChannel::DeleteKeysFrom(FFrameNumber InTime, bool
 	GetData().DeleteKeysFrom(InTime, bDeleteKeysBefore);
 }
 
-void FMovieSceneNiagaraEmitterChannel::ChangeFrameResolution(FFrameRate SourceRate, FFrameRate DestinationRate)
+void FMovieSceneNiagaraEmitterChannel::RemapTimes(const UE::MovieScene::IRetimingInterface& Retimer)
 {
-	GetData().ChangeFrameResolution(SourceRate, DestinationRate);
-}
-
-void FMovieSceneNiagaraEmitterChannel::RemapMusicalFrameTimes(const UFrameBasedMusicMap* OldMap, const UFrameBasedMusicMap* NewMap)
-{
-	GetData().RemapMusicalFrameTimes(OldMap, NewMap);
+	GetData().RemapTimes(Retimer);
 }
 
 TRange<FFrameNumber> FMovieSceneNiagaraEmitterChannel::ComputeEffectiveRange() const
diff --git a/Engine/Plugins/FX/Niagara/Source/NiagaraEditor/Private/Sequencer/NiagaraSequence/Sections/MovieSceneNiagaraEmitterSection.h b/Engine/Plugins/FX/Niagara/Source/NiagaraEditor/Private/Sequencer/NiagaraSequence/Sections/MovieSceneNiagaraEmitterSection.h
--- a/Engine/Plugins/FX/Niagara/Source/NiagaraEditor/Private/Sequencer/NiagaraSequence/Sections/MovieSceneNiagaraEmitterSection.h
+++ b/Engine/Plugins/FX/Niagara/Source/NiagaraEditor/Private/Sequencer/NiagaraSequence/Sections/MovieSceneNiagaraEmitterSection.h
@@ -96,8 +96,7 @@ public:
 	virtual void DuplicateKeys(TArrayView<const FKeyHandle> InHandles, TArrayView<FKeyHandle> OutNewHandles) override;
 	virtual void DeleteKeys(TArrayView<const FKeyHandle> InHandles) override;
 	virtual void DeleteKeysFrom(FFrameNumber InTime, bool bDeleteKeysBefore) override;
-	virtual void ChangeFrameResolution(FFrameRate SourceRate, FFrameRate DestinationRate) override;
-	virtual void RemapMusicalFrameTimes(const UFrameBasedMusicMap* OldMap, const UFrameBasedMusicMap* NewMap) override;
+	virtual void RemapTimes(const UE::MovieScene::IRetimingInterface& Retimer) override;
 	virtual TRange<FFrameNumber> ComputeEffectiveRange() const override;
 	virtual int32 GetNumKeys() const override;
 	virtual void Reset() override;
diff --git a/Engine/Plugins/MovieScene/LevelSequenceEditor/Source/LevelSequenceEditor/Private/CinematicViewport/SCinematicLevelViewport.cpp b/Engine/Plugins/MovieScene/LevelSequenceEditor/Source/LevelSequenceEditor/Private/CinematicViewport/SCinematicLevelViewport.cpp
--- a/Engine/Plugins/MovieScene/LevelSequenceEditor/Source/LevelSequenceEditor/Private/CinematicViewport/SCinematicLevelViewport.cpp
+++ b/Engine/Plugins/MovieScene/LevelSequenceEditor/Source/LevelSequenceEditor/Private/CinematicViewport/SCinematicLevelViewport.cpp
@@ -767,8 +767,8 @@ float SCinematicLevelViewport::GetPlayTimeMinDesiredWidth() const
 	{
 		TRange<double> ViewRange = Sequencer->GetViewRange();
 
-		FString LowerBoundStr = Sequencer->GetNumericPositionTypeInterface()->ToString(ViewRange.GetLowerBoundValue());
-		FString UpperBoundStr = Sequencer->GetNumericPositionTypeInterface()->ToString(ViewRange.GetUpperBoundValue());
+		FString LowerBoundStr = Sequencer->GetNumericTypeInterface()->ToString(ViewRange.GetLowerBoundValue());
+		FString UpperBoundStr = Sequencer->GetNumericTypeInterface()->ToString(ViewRange.GetUpperBoundValue());
 
 		// Always measure with the negative and subframe indicator so that the size doesn't change when there is and isn't a subframe
 		if (!LowerBoundStr.Contains(TEXT("*")))
@@ -864,7 +864,7 @@ void SCinematicLevelViewport::Setup(FLevelSequenceEditorToolkit& NewToolkit)
 	ISequencer* Sequencer = GetSequencer();
 	if (Sequencer)
 	{
-		TypeInterfaceProxy->Impl = Sequencer->GetNumericPositionTypeInterface();
+		TypeInterfaceProxy->Impl = Sequencer->GetNumericTypeInterface();
 
 		if (TransportRange.IsValid())
 		{
@@ -964,7 +964,7 @@ void SCinematicLevelViewport::Tick(const FGeometry& AllottedGeometry, const doub
 
 	FText TimeFormat = LOCTEXT("TimeFormat", "{0}");
 
-	TSharedPtr<INumericTypeInterface<double>> TimeDisplayFormatInterface = Sequencer->GetNumericPositionTypeInterface();
+	TSharedPtr<INumericTypeInterface<double>> TimeDisplayFormatInterface = Sequencer->GetNumericTypeInterface();
 
 	UMovieSceneSequence* SubSequence = CinematicShotSection ? CinematicShotSection->GetSequence() : nullptr;
 	if (SubSequence)
diff --git a/Engine/Plugins/MovieScene/LevelSequenceEditor/Source/LevelSequenceEditor/Private/LevelSequenceEditorSubsystem.cpp b/Engine/Plugins/MovieScene/LevelSequenceEditor/Source/LevelSequenceEditor/Private/LevelSequenceEditorSubsystem.cpp
--- a/Engine/Plugins/MovieScene/LevelSequenceEditor/Source/LevelSequenceEditor/Private/LevelSequenceEditorSubsystem.cpp
+++ b/Engine/Plugins/MovieScene/LevelSequenceEditor/Source/LevelSequenceEditor/Private/LevelSequenceEditorSubsystem.cpp
@@ -1473,9 +1473,9 @@ void SBakeTransformWidget::Construct(const FArguments& InArgs)
 	FPropertyEditorModule& PropertyEditor = FModuleManager::Get().LoadModuleChecked<FPropertyEditorModule>(TEXT("PropertyEditor"));
 
 	DetailsView = PropertyEditor.CreateStructureDetailView(ViewArgs, StructureViewArgs, TSharedPtr<FStructOnScope>());
-	TSharedPtr<INumericTypeInterface<double>> NumericTypeInterface = Sequencer->GetNumericPositionTypeInterface();
+
 	DetailsView->GetDetailsView()->RegisterInstancedCustomPropertyTypeLayout("FrameNumber",
-		FOnGetPropertyTypeCustomizationInstance::CreateLambda([=]() {return MakeShared<FFrameNumberDetailsCustomization>(NumericTypeInterface); }));
+		FOnGetPropertyTypeCustomizationInstance::CreateSP(Sequencer, &ISequencer::MakeFrameNumberDetailsCustomization));
 	DetailsView->SetStructureData(Settings);
 
 	ChildSlot
diff --git a/Engine/Plugins/MovieScene/MovieSceneTextTrack/Source/MovieSceneTextTrack/Private/MovieSceneTextChannel.cpp b/Engine/Plugins/MovieScene/MovieSceneTextTrack/Source/MovieSceneTextTrack/Private/MovieSceneTextChannel.cpp
--- a/Engine/Plugins/MovieScene/MovieSceneTextTrack/Source/MovieSceneTextTrack/Private/MovieSceneTextChannel.cpp
+++ b/Engine/Plugins/MovieScene/MovieSceneTextTrack/Source/MovieSceneTextTrack/Private/MovieSceneTextChannel.cpp
@@ -69,14 +69,9 @@ void FMovieSceneTextChannel::DeleteKeysFrom(FFrameNumber InTime, bool bDeleteKey
 	GetData().DeleteKeysFrom(InTime, bDeleteKeysBefore);
 }
 
-void FMovieSceneTextChannel::ChangeFrameResolution(FFrameRate SourceRate, FFrameRate DestinationRate)
+void FMovieSceneTextChannel::RemapTimes(const UE::MovieScene::IRetimingInterface& Retimer)
 {
-	GetData().ChangeFrameResolution(SourceRate, DestinationRate);
-}
-
-void FMovieSceneTextChannel::RemapMusicalFrameTimes(const UFrameBasedMusicMap* OldMap, const UFrameBasedMusicMap* NewMap)
-{
-	GetData().RemapMusicalFrameTimes(OldMap, NewMap);
+	GetData().RemapTimes(Retimer);
 }
 
 TRange<FFrameNumber> FMovieSceneTextChannel::ComputeEffectiveRange() const
diff --git a/Engine/Plugins/MovieScene/MusicalMode/MusicalMode.uplugin b/Engine/Plugins/MovieScene/MusicalMode/MusicalMode.uplugin
new file mode 100644
--- /dev/null
+++ b/Engine/Plugins/MovieScene/MusicalMode/MusicalMode.uplugin
@@ -0,0 +1,30 @@
+{
+	"FileVersion": 1,
+	"Version": 1,
+	"VersionName": "1.0",
+	"FriendlyName": "Sequencer Musical Mode",
+	"Description": "Adds the ability for Sequencer to be placed into a musical mode",
+	"Category": "Text",
+	"CreatedBy": "Epic Games, Inc.",
+	"CreatedByURL": "https://epicgames.com",
+	"DocsURL": "",
+	"MarketplaceURL": "",
+	"SupportURL": "",
+	"CanContainContent": false,
+	"IsBetaVersion": true,
+	"IsExperimentalVersion": true,
+	"EnabledByDefault": false,
+	"Installed": false,
+	"Modules": [
+		{
+			"Name": "MusicalMode",
+			"Type": "Runtime",
+			"LoadingPhase": "Default"
+		},
+		{
+			"Name": "MusicalModeEditor",
+			"Type": "Editor",
+			"LoadingPhase": "Default"
+		}
+	]
+}
diff --git a/Engine/Plugins/MovieScene/MusicalMode/Source/MusicalMode/MusicalMode.build.cs b/Engine/Plugins/MovieScene/MusicalMode/Source/MusicalMode/MusicalMode.build.cs
new file mode 100644
--- /dev/null
+++ b/Engine/Plugins/MovieScene/MusicalMode/Source/MusicalMode/MusicalMode.build.cs
@@ -0,0 +1,23 @@
+// Copyright Epic Games, Inc. All Rights Reserved.
+
+using UnrealBuildTool;
+
+public class MusicalMode : ModuleRules
+{
+	public MusicalMode(ReadOnlyTargetRules Target) : base(Target)
+	{
+		PCHUsage = ModuleRules.PCHUsageMode.UseExplicitOrSharedPCHs;
+
+		PrivateDependencyModuleNames.AddRange(
+			new string[]
+			{
+				"Core",
+				"CoreUObject",
+				"Engine",
+				"MovieScene",
+				"TimeManagement",
+				"MusicEnvironment",
+			}
+		);
+	}
+}
diff --git a/Engine/Plugins/MovieScene/MusicalMode/Source/MusicalMode/Private/MetronomeTimeControllers.cpp b/Engine/Plugins/MovieScene/MusicalMode/Source/MusicalMode/Private/MetronomeTimeControllers.cpp
new file mode 100644
--- /dev/null
+++ b/Engine/Plugins/MovieScene/MusicalMode/Source/MusicalMode/Private/MetronomeTimeControllers.cpp
@@ -0,0 +1,88 @@
+// Copyright Epic Games, Inc. All Rights Reserved.
+
+#include "MetronomeTimeControllers.h"
+#include "MovieSceneMusicalClock.h"
+#include "Misc/MusicalTime.h"
+
+namespace MovieSceneMusicalMode
+{
+	FMetronomeTimeController::~FMetronomeTimeController()
+	{
+		// It may take a while before the UObject we are referencing actually gets GC'd, so
+		// be sure it isn't still playing a metronome sound while it waits...
+		if (MetronomeInterfacePtr)
+		{
+			MetronomeInterfacePtr->Stop();
+		}
+		MetronomeInterfacePtr = nullptr;
+		MetronomeObjectPtr = nullptr;
+	}
+
+	void FMetronomeTimeController::OnStartPlaying(const FQualifiedFrameTime& InStartTime)
+	{
+		check(MetronomeInterfacePtr);
+		if (TStrongObjectPtr<const UMovieSceneMusicalClock> Clock = MusicalMovieSceneClock.Pin())
+		{
+			float FromSeconds = static_cast<float>(InStartTime.AsSeconds() * Clock->MusicMap->GetInitialTempo() / MetronomeInterfacePtr->GetCurrentTempo());
+			MetronomeInterfacePtr->Start(FromSeconds);
+		}
+		else
+		{
+			MetronomeInterfacePtr->Start(InStartTime.AsSeconds());
+		}
+	}
+
+	void FMetronomeTimeController::OnTick(float DeltaSeconds, float InPlayRate)
+	{
+		check(MetronomeInterfacePtr);
+		MetronomeInterfacePtr->Tick(DeltaSeconds);
+	}
+
+	void FMetronomeTimeController::OnStopPlaying(const FQualifiedFrameTime& InStopTime)
+	{
+		check(MetronomeInterfacePtr);
+		MetronomeInterfacePtr->Stop();
+	}
+
+	FFrameTime FMetronomeTimeController::OnRequestCurrentTime(const FQualifiedFrameTime& InCurrentTime, float InPlayRate)
+	{
+		check(MetronomeInterfacePtr);
+		if (TStrongObjectPtr<const UMovieSceneMusicalClock> Clock = MusicalMovieSceneClock.Pin())
+		{
+			MetronomeInterfacePtr->SetTempo(Clock->GetAuditionTempo());
+			FMusicalTime MetronomePosition = MetronomeInterfacePtr->GetMusicalTime();
+			FFrameTime MovieResFrameTime = Clock->MusicMap->GetFrameTime(MetronomePosition);
+			FFrameTime ConvertedTime = ConvertFrameTime(MovieResFrameTime, Clock->MusicMap->FrameResolution, InCurrentTime.Rate);  
+			return ConvertedTime;
+		}
+		return FFrameTime();
+	}
+
+	FFrameTime FSimpleSilentMusicalTimeController::OnRequestCurrentTime(const FQualifiedFrameTime& InCurrentTime, float InPlayRate)
+	{
+		LastRealWallClockTime = FMovieSceneTimeController_PlatformClock::OnRequestCurrentTime(InCurrentTime, InPlayRate).AsDecimal();
+		double DeltaTimeAtCurrentTempo = LastRealWallClockTime - RealWallClockTimeOfLastTempoChange;
+		double TimeScale = (CurrentTempo / AuthoredTempo);
+		ElapsedTimeAtTempo = DeltaTimeAtCurrentTempo * TimeScale;
+		LastReportedTime = ReportedTimeAtLastTempoChange + ElapsedTimeAtTempo;
+		return FFrameTime::FromDecimal(LastReportedTime);
+	}
+
+	void FSimpleSilentMusicalTimeController::OnSetAuditionTempo(float Tempo)
+	{
+		CurrentTempo = Tempo;
+		RealWallClockTimeOfLastTempoChange = LastRealWallClockTime;
+		ReportedTimeAtLastTempoChange = LastReportedTime;
+		ElapsedTimeAtTempo = 0.0;
+	}
+
+	void FSimpleSilentMusicalTimeController::OnStartPlaying(const FQualifiedFrameTime& InStartTime)
+	{
+		FMovieSceneTimeController_PlatformClock::OnStartPlaying(InStartTime);
+		LastRealWallClockTime = FMovieSceneTimeController_PlatformClock::OnRequestCurrentTime(InStartTime, 1.0f).AsDecimal();
+		RealWallClockTimeOfLastTempoChange = LastRealWallClockTime;
+		ReportedTimeAtLastTempoChange = LastRealWallClockTime;
+		LastReportedTime = LastRealWallClockTime;
+		ElapsedTimeAtTempo = 0.0;
+	}
+}
\ No newline at end of file
diff --git a/Engine/Plugins/MovieScene/MusicalMode/Source/MusicalMode/Private/MovieSceneMusicalClock.cpp b/Engine/Plugins/MovieScene/MusicalMode/Source/MusicalMode/Private/MovieSceneMusicalClock.cpp
new file mode 100644
--- /dev/null
+++ b/Engine/Plugins/MovieScene/MusicalMode/Source/MusicalMode/Private/MovieSceneMusicalClock.cpp
@@ -0,0 +1,166 @@
+// Copyright Epic Games, Inc. All Rights Reserved.
+
+#include "MovieSceneMusicalClock.h"
+
+#include "Engine/World.h"
+#include "FrameBasedMusicMap.h"
+#include "MovieScene.h"
+#include "MovieSceneMusicTimeTransform.h"
+#include "MetronomeTimeControllers.h"
+#include "Sections/MovieSceneSectionTimingParameters.h"
+#include "Sections/MovieSceneSubSection.h"
+#include "MusicEnvironmentSubsystem.h"
+
+TSharedPtr<FMovieSceneTimeController> UMovieSceneMusicalClock::MakeTimeController(UObject* PlaybackContext) const
+{
+	UWorld* ContextWorld = PlaybackContext ? PlaybackContext->GetWorld() : nullptr;
+
+	if (ContextWorld && ( ContextWorld->WorldType == EWorldType::PIE || ContextWorld->WorldType == EWorldType::Game) )
+	{
+#if 0
+		UMusicClockComponent* ClockComponent = Cast<UMusicClockComponent>(ClockComponentPath.ResolveObject());
+		if (ClockComponent)
+		{
+			struct FMusicalTimeClock : FMovieSceneTimeController
+			{
+				TWeakObjectPtr<UMusicClockComponent> WeakClockComponent;
+				TWeakObjectPtr<UFrameBasedMusicMap> WeakMusicMap;
+
+				FMusicalTimeClock(UMusicClockComponent* ClockComponent, UFrameBasedMusicMap* Map)
+					: WeakClockComponent(ClockComponent)
+					, WeakMusicMap(Map)
+				{}
+
+				FFrameTime OnRequestCurrentTime(const FQualifiedFrameTime& InCurrentTime, float InPlayRate) override
+				{
+					UMusicClockComponent* ClockComponent = WeakClockComponent.Get();
+					if (ClockComponent)
+					{
+						FMusicalTime MusicalTime = ClockComponent->OnRequestCurrentMusicalTime(InCurrentTime, InPlayRate);
+						FFrameTime FrameTime = WeakMusicMap.Get()->GetFrameTime(MusicalTime);
+						return NewTime = FFrameRate::TransformTime(FrameTime, WeakMusicMap.Get(), GetDisplayRate());
+					}
+
+					return FFrameTime();
+				}
+
+			};
+			// Cast it to some kind of musical clock interface????
+			return MakeShared<FMusicalTimeClock>(ClockComponent);
+		}
+#endif
+	}
+
+	if (TScriptInterface<IMusicEnvironmentMetronome> Metronome = UMusicEnvironmentSubsystem::Get().SpawnMetronome(PlaybackContext))
+	{
+		// Great... We have a nice metronome that will presumably beep along at the audition tempo.
+		Metronome->Initialize(this->GetWorld());
+		Metronome->SetMusicMap(MusicMap);
+		return MakeShared<MovieSceneMusicalMode::FMetronomeTimeController>(Metronome, this);
+	}
+		
+	// If we got here we couldn't spawn an audible metronome with the Music Environment Subsystem, so we will
+	// create a simple time controller that will at least respond to audition tempo settings. It won't beep
+	// or anything, but it is better than nothing.  
+	return MakeShared<MovieSceneMusicalMode::FSimpleSilentMusicalTimeController>(MusicMap->GetInitialTempo(), 120.0f);
+}
+
+void UMovieSceneMusicalClock::HandleTickResolutionChange(FFrameRate PreviousTickResolution, FFrameRate NewTickResolution)
+{
+	if (MusicMap)
+	{
+		MusicMap->SetFrameResolution(NewTickResolution);
+	}
+}
+
+bool UMovieSceneMusicalClock::MakeSubSequenceTransform(const FMovieSceneSectionTimingParametersFrames& Timing, const UMovieSceneSubSection* SubSection, FMovieSceneSequenceTransform& OutTransform) const
+{
+	UMovieScene* OuterMovieScene = SubSection->GetTypedOuter<UMovieScene>();
+	UMovieScene* InnerMovieScene = GetTypedOuter<UMovieScene>();
+
+	TRange<FFrameNumber> OuterRange = SubSection->GetRange();
+	TRange<FFrameNumber> InnerRange = GetTypedOuter<UMovieScene>()->GetPlaybackRange();
+
+	check(OuterRange.HasLowerBound());
+	check(InnerRange.HasLowerBound() && InnerRange.HasUpperBound());
+	check(OuterMovieScene);
+	check(InnerMovieScene);
+
+	UMovieSceneMusicalClock* OuterClock = OuterMovieScene->GetCustomClockAs<UMovieSceneMusicalClock>();
+	UMovieSceneMusicalClock* InnerClock = OuterMovieScene->GetCustomClockAs<UMovieSceneMusicalClock>();
+	if (!OuterClock || !InnerClock || !OuterClock->MusicMap || !InnerClock->MusicMap)
+	{
+		return false;
+	}
+
+	const UFrameBasedMusicMap* OuterMap = OuterClock->MusicMap;
+	const UFrameBasedMusicMap* InnerMap = InnerClock->MusicMap;
+	check(OuterMap);
+	check(InnerMap);
+
+	// ----------------------------------------------------------------------------
+	// First things first, subtract the section start bound
+	Timing.AddPositionInOuterAsOffset(OutTransform, -OuterRange.GetLowerBoundValue());
+
+	// ----------------------------------------------------------------------------
+	// FrameRate conversion
+	if (InnerMovieScene->GetTickResolution() != OuterMovieScene->GetTickResolution())
+	{
+		Timing.AddFrameRateConversion(OutTransform, OuterMovieScene->GetTickResolution(), InnerMovieScene->GetTickResolution());
+	}
+
+	// ----------------------------------------------------------------------------
+	// Musical transform
+	UMovieSceneMusicTimeTransform* Mapper = NewObject<UMovieSceneMusicTimeTransform>(InnerMovieScene);
+	Mapper->Initialize(OuterClock, InnerClock, OuterRange.GetLowerBoundValue());
+	OutTransform.Add(0, FMovieSceneTimeWarpVariant(Mapper));
+
+	// ----------------------------------------------------------------------------
+	// Time warp
+	Timing.AddPlayRate(OutTransform);
+
+
+	// InnerStartOffset and InnerEndOffset are in frames that are both...
+	// a. In units of the outer music map, and...
+	// b. a count from the BEGINNING of the outer music map as a duration.
+	// So we have some conversions to use...
+	FMusicalTime StartOffsetAsMusicalTime = OuterMap->GetMusicalTime(Timing.InnerStartOffset);
+	FFrameNumber RealInnerStartOffset = InnerMap->GetFrameTime(StartOffsetAsMusicalTime).RoundToFrame();
+	FMusicalTime EndOffsetAsMusicalTime = OuterMap->GetMusicalTime(Timing.InnerEndOffset);
+	FFrameNumber RealInnerEndOffset = InnerMap->GetFrameTime(EndOffsetAsMusicalTime).RoundToFrame();
+
+	FMusicalTime LoopStartOffsetAsMusicalTime = OuterMap->GetMusicalTime(Timing.FirstLoopStartOffset);
+	FFrameNumber RealFirstLoopStartOffset = InnerMap->GetFrameTime(LoopStartOffsetAsMusicalTime).RoundToFrame();
+
+
+	FFrameNumber StartTime = InnerRange.GetLowerBoundValue() + RealInnerStartOffset;
+	FFrameNumber EndTime = InnerRange.GetUpperBoundValue() - RealInnerEndOffset;
+	FFrameNumber Duration = EndTime - StartTime;
+
+	FFrameNumber LoopOffset(Timing.bLoop ? RealFirstLoopStartOffset : 0);
+
+	FFrameNumber NegativeRateOffset = 0;
+	if (Timing.PlayRate.GetType() == EMovieSceneTimeWarpType::FixedPlayRate && Timing.PlayRate.AsFixedPlayRate() < 0.0)
+	{
+		NegativeRateOffset = Duration;
+	}
+
+	// Start offset
+	Timing.AddInnerStartOffset(OutTransform, StartTime + LoopOffset + NegativeRateOffset);
+
+	// ----------------------------------------------------------------------------
+	// Looping or clamping
+	if (Duration > 0)
+	{
+		Timing.AddLoopingOrClampingAndReverse(OutTransform, Duration, -StartTime);
+	}
+
+	return true;
+}
+
+void UMovieSceneMusicalClock::CreateSimpleMusicMap(float InBpm, const FFrameBasedTimeSignature& InTimeSignature)
+{
+	MusicMap = NewObject<UFrameBasedMusicMap>(this, NAME_None, RF_Transactional);
+	MusicMap->SetFrameResolution(this->GetTypedOuter<UMovieScene>()->GetTickResolution());
+	MusicMap->Init(120.0f, { 4, 4 });
+}
diff --git a/Engine/Plugins/MovieScene/MusicalMode/Source/MusicalMode/Private/MusicalModeModule.cpp b/Engine/Plugins/MovieScene/MusicalMode/Source/MusicalMode/Private/MusicalModeModule.cpp
new file mode 100644
--- /dev/null
+++ b/Engine/Plugins/MovieScene/MusicalMode/Source/MusicalMode/Private/MusicalModeModule.cpp
@@ -0,0 +1,8 @@
+// Copyright Epic Games, Inc. All Rights Reserved.
+
+#include "Modules/ModuleInterface.h"
+#include "Modules/ModuleManager.h"
+
+
+
+IMPLEMENT_MODULE(FDefaultModuleImpl, MusicalMode);
diff --git a/Engine/Plugins/MovieScene/MusicalMode/Source/MusicalMode/Public/MetronomeTimeControllers.h b/Engine/Plugins/MovieScene/MusicalMode/Source/MusicalMode/Public/MetronomeTimeControllers.h
new file mode 100644
--- /dev/null
+++ b/Engine/Plugins/MovieScene/MusicalMode/Source/MusicalMode/Public/MetronomeTimeControllers.h
@@ -0,0 +1,63 @@
+// Copyright Epic Games, Inc. All Rights Reserved.
+
+#pragma once
+
+#include "MovieSceneTimeController.h"
+#include "MovieSceneMusicalClock.h"
+#include "MusicEnvironmentMetronome.h"
+
+namespace MovieSceneMusicalMode
+{
+	/**
+	 * 
+	 */
+	struct FMetronomeTimeController : FMovieSceneTimeController_PlatformClock
+	{
+	public:
+		FMetronomeTimeController(const TScriptInterface<IMusicEnvironmentMetronome>& Metronome, const UMovieSceneMusicalClock* Clock)
+			: MetronomeObjectPtr(TStrongObjectPtr<UObject>(Metronome.GetObject()))
+			, MetronomeInterfacePtr(Metronome.GetInterface())
+			, MusicalMovieSceneClock(Clock)
+		{}
+
+		~FMetronomeTimeController() override;
+	
+		virtual void OnStartPlaying(const FQualifiedFrameTime& InStartTime) override;
+		virtual void OnTick(float DeltaSeconds, float InPlayRate) override;
+		virtual void OnStopPlaying(const FQualifiedFrameTime& InStopTime) override;
+		virtual FFrameTime OnRequestCurrentTime(const FQualifiedFrameTime& InCurrentTime, float InPlayRate) override;
+
+	private:
+		TStrongObjectPtr<UObject> MetronomeObjectPtr;
+		IMusicEnvironmentMetronome* MetronomeInterfacePtr;
+		TWeakObjectPtr<const UMovieSceneMusicalClock> MusicalMovieSceneClock;
+	};
+
+	struct FSimpleSilentMusicalTimeController : FMovieSceneTimeController_PlatformClock
+	{
+		FSimpleSilentMusicalTimeController(float InAuthoredTempo, float InCurrentTempo)
+			: AuthoredTempo(InAuthoredTempo)
+			, CurrentTempo(InCurrentTempo)
+			, LastRealWallClockTime(0.0)
+			, RealWallClockTimeOfLastTempoChange(0.0)
+			, ReportedTimeAtLastTempoChange(0.0)
+			, ElapsedTimeAtTempo(0.0)
+			, LastReportedTime(0.0)
+		{}
+
+		virtual FFrameTime OnRequestCurrentTime(const FQualifiedFrameTime& InCurrentTime, float InPlayRate) override;
+
+		virtual void OnStartPlaying(const FQualifiedFrameTime& InStartTime) override;
+
+		void OnSetAuditionTempo(float Tempo);
+	
+	private:
+		float AuthoredTempo;
+		float CurrentTempo;
+		double LastRealWallClockTime;
+		double RealWallClockTimeOfLastTempoChange;
+		double ReportedTimeAtLastTempoChange;
+		double ElapsedTimeAtTempo;
+		double LastReportedTime;
+	};
+}
\ No newline at end of file
diff --git a/Engine/Plugins/MovieScene/MusicalMode/Source/MusicalMode/Public/MovieSceneMusicalClock.h b/Engine/Plugins/MovieScene/MusicalMode/Source/MusicalMode/Public/MovieSceneMusicalClock.h
new file mode 100644
--- /dev/null
+++ b/Engine/Plugins/MovieScene/MusicalMode/Source/MusicalMode/Public/MovieSceneMusicalClock.h
@@ -0,0 +1,44 @@
+// Copyright Epic Games, Inc. All Rights Reserved.
+
+#pragma once
+
+#include "FrameBasedMusicMap.h"
+#include "MovieSceneClock.h"
+#include "MovieSceneMusicalClock.generated.h"
+
+class UFrameBasedMusicMap;
+
+UCLASS(MinimalAPI)
+class UMovieSceneMusicalClock : public UMovieSceneClock
+{
+public:
+
+	GENERATED_BODY()
+
+	virtual TSharedPtr<FMovieSceneTimeController> MakeTimeController(UObject* PlaybackContext) const override;
+	virtual void HandleTickResolutionChange(FFrameRate PreviousTickResolution, FFrameRate NewTickResolution) override;
+	virtual bool MakeSubSequenceTransform(const FMovieSceneSectionTimingParametersFrames& Timing, const UMovieSceneSubSection* SubSection, FMovieSceneSequenceTransform& OutTransform) const;
+
+	UPROPERTY()
+	TObjectPtr<UFrameBasedMusicMap> MusicMap;
+
+	UPROPERTY()
+	FSoftObjectPath ClockComponentPath;
+
+	UPROPERTY()
+	EUpdateClockSource FallbackClock = EUpdateClockSource::Audio;
+
+	UPROPERTY()
+	bool bDisplayFramesAsBarsBeats = true;
+
+	void MUSICALMODE_API CreateSimpleMusicMap(float InBpm, const FFrameBasedTimeSignature& InTimeSignature);
+	
+	void SetAuditionTempo(float InAuditionTempo) { AuditionTempo = InAuditionTempo; }
+
+	float GetAuditionTempo() const { return AuditionTempo < 1.0f ? MusicMap->GetInitialTempo() : AuditionTempo; }
+
+private:
+	UPROPERTY()
+	float AuditionTempo = 0.0f; // if < 1.0 then audition at music map's initial tempo. 	
+};
+
diff --git a/Engine/Plugins/MovieScene/MusicalMode/Source/MusicalModeEditor/MusicalModeEditor.build.cs b/Engine/Plugins/MovieScene/MusicalMode/Source/MusicalModeEditor/MusicalModeEditor.build.cs
new file mode 100644
--- /dev/null
+++ b/Engine/Plugins/MovieScene/MusicalMode/Source/MusicalModeEditor/MusicalModeEditor.build.cs
@@ -0,0 +1,33 @@
+// Copyright Epic Games, Inc. All Rights Reserved.
+
+using UnrealBuildTool;
+
+public class MusicalModeEditor : ModuleRules
+{
+	public MusicalModeEditor(ReadOnlyTargetRules Target) : base(Target)
+	{
+		PCHUsage = ModuleRules.PCHUsageMode.UseExplicitOrSharedPCHs;
+
+		PrivateDependencyModuleNames.AddRange(
+			new string[]
+			{
+				"Core",
+				"CoreUObject",
+				"EditorStyle",
+				"Engine",
+				"MovieScene",
+				"MusicalMode",
+				"Sequencer",
+				"SequencerCore",
+				"Slate",
+				"SlateCore",
+				"InputCore",
+				"TimeManagement",
+				"MusicEnvironment",
+				"ToolMenus",
+				"ToolWidgets",
+				"UnrealEd",
+			}
+		);
+	}
+}
diff --git a/Engine/Plugins/MovieScene/MusicalMode/Source/MusicalModeEditor/Private/MusicalModeClockExtension.cpp b/Engine/Plugins/MovieScene/MusicalMode/Source/MusicalModeEditor/Private/MusicalModeClockExtension.cpp
new file mode 100644
--- /dev/null
+++ b/Engine/Plugins/MovieScene/MusicalMode/Source/MusicalModeEditor/Private/MusicalModeClockExtension.cpp
@@ -0,0 +1,419 @@
+// Copyright Epic Games, Inc. All Rights Reserved.
+
+#include "MusicalModeClockExtension.h"
+#include "Fonts/FontMeasure.h"
+#include "FrameBasedMusicMap.h"
+#include "FrameNumberDisplayFormat.h"
+#include "Framework/Application/SlateApplication.h"
+#include "ISequencerNumericTypeInterface.h"
+#include "Misc/AsciiSet.h"
+#include "MovieScene.h"
+#include "MovieSceneMusicalClock.h"
+#include "MovieSceneSequence.h"
+#include "MVVM/ViewModels/SequenceModel.h"
+#include "PropertyHandle.h"
+#include "SequencerSettings.h"
+#include "Styling/CoreStyle.h"
+
+DEFINE_EXPRESSION_OPERATOR_NODE(, FBarBeatSeparator, 0xb2044e17, 0x81c9429c, 0xac0d7398, 0x9d5cb582);
+const TCHAR* const ExpressionParser::FBarBeatSeparator::Moniker = TEXT(":");
+
+namespace UE::Sequencer
+{
+
+struct FBarBeatStringParser
+{
+	FBarBeatStringParser()
+	{
+		using namespace ExpressionParser;
+
+		TokenDefinitions.IgnoreWhitespace();
+		TokenDefinitions.DefineToken(&ConsumeSymbol<FBarBeatSeparator>);
+		TokenDefinitions.DefineToken(&ConsumeSymbol<FPlus>);
+		TokenDefinitions.DefineToken(&ConsumeSymbol<FMinus>);
+		TokenDefinitions.DefineToken(&ConsumeLocalizedNumber);
+
+		Grammar.DefineBinaryOperator<FPlus>(5);
+		Grammar.DefineBinaryOperator<FStar>(4);
+
+		// @todo: Not sure what is correct here if the user types something like 1:1 + 2
+		JumpTable.MapBinary<FPlus>([](double A, double B) { return A + B; });
+		JumpTable.MapBinary<FStar>([](double A, double B) { return A * B; });
+	}
+
+	TValueOrError<FFrameTime, FExpressionError> EvaluateBarBeat(const TCHAR* InExpression, const UFrameBasedMusicMap* MusicMap, FrameBasedMusicMap::EBarBeatFormat BarBeatFormat) const
+	{
+		if (!MusicMap || MusicMap->IsEmpty())
+		{
+			return MakeError(NSLOCTEXT("MusicMode", "BarBeatError_NoTempoMap", "No supplied TempoMap"));
+		}
+
+		FString InAsString(InExpression);
+
+		// find bar.beat boundary and replace with space...
+		int32 PeriodCount = 0;
+		for (TCHAR& Character : InAsString)
+		{
+			if (Character == '.')
+			{
+				PeriodCount++;
+				if (PeriodCount == 1)
+				{
+					Character = ' ';
+				}
+			}
+		}
+
+		FText Error = NSLOCTEXT("MusicMode", "BarBeatError_BadFormat", "Format not recognized as a Bar.Beat");
+		if (PeriodCount > 2)
+		{
+			return MakeError(Error);
+		}
+
+		TValueOrError<TArray<FExpressionToken>, FExpressionError> LexResult = ExpressionParser::Lex(*InAsString, TokenDefinitions);
+		if (!LexResult.IsValid())
+		{
+			return MakeError(LexResult.StealError());
+		}
+		TArray<FExpressionToken> Tokens = LexResult.StealValue();
+		if (Tokens.Num() == 0 || Tokens.Num() > 2)
+		{
+			return MakeError(Error);
+		}
+
+		const double* BarToken = Tokens[0].Node.Cast<double>();
+		if (!BarToken)
+		{
+			return MakeError(Error);
+		}
+
+		const double* BeatToken = nullptr;
+		if (Tokens.Num() == 2)
+		{
+			BeatToken = Tokens[1].Node.Cast<double>();
+			if (!BeatToken)
+			{
+				return MakeError(Error);
+			}
+			else if (*BeatToken < 1.0f && BarBeatFormat == FrameBasedMusicMap::EBarBeatFormat::Position)
+			{
+				return MakeError(NSLOCTEXT("MusicMode", "BarBeatError_InvalidBeatValue", "Beats are one based, so the specified beat must be >= 1.0"));
+			}
+		}
+
+		int32 Bar = static_cast<int32>(*BarToken);
+		float BeatInBar = static_cast<float>(BeatToken ? *BeatToken : 1.0);
+		
+		// in user land bars and beats are 1 based, so make them 0 based.
+		int32 BarCorrection = BarBeatFormat == FrameBasedMusicMap::EBarBeatFormat::Position ? 1 : 0;
+		float BeatCorrection = BarBeatFormat == FrameBasedMusicMap::EBarBeatFormat::Position ? 1.0f : 0.0f;
+		return MakeValue(MusicMap->GetFrameTime(Bar - BarCorrection, BeatInBar - BeatCorrection));
+	}
+
+	FTokenDefinitions TokenDefinitions;
+	FExpressionGrammar Grammar;
+
+	FOperatorJumpTable JumpTable;
+};
+
+
+struct FBarsBeatsNumericTypeInterface : public INumericTypeInterface<double>
+{
+	TWeakPtr<ISequencer> WeakSequencer;
+	ENumericIntent Intent;
+
+	FBarsBeatsNumericTypeInterface(TWeakPtr<ISequencer> InWeakSequencer, ENumericIntent InIntent)
+		: WeakSequencer(InWeakSequencer)
+		, Intent(InIntent)
+	{}
+
+	/** Gets the minimum and maximum fractional digits. */
+	virtual int32 GetMinFractionalDigits() const override
+	{
+		return 0;
+	}
+	virtual int32 GetMaxFractionalDigits() const override
+	{
+		return 0;
+	}
+	virtual void SetMinFractionalDigits(const TAttribute<TOptional<int32>>& NewValue) override
+	{
+	}
+	virtual void SetMaxFractionalDigits(const TAttribute<TOptional<int32>>& NewValue) override
+	{
+	}
+	virtual bool IsCharacterValid(TCHAR InChar) const override
+	{
+		auto IsValidLocalizedCharacter = [InChar]() -> bool
+		{
+			const FDecimalNumberFormattingRules& NumberFormattingRules = ExpressionParser::GetLocalizedNumberFormattingRules();
+			return InChar == NumberFormattingRules.GroupingSeparatorCharacter
+				|| InChar == NumberFormattingRules.DecimalSeparatorCharacter
+				|| Algo::Find(NumberFormattingRules.DigitCharacters, InChar) != 0;
+		};
+
+		static const FAsciiSet ValidChars = TEXT("1234567890:");
+		return InChar != 0 && (ValidChars.Contains(InChar) || IsValidLocalizedCharacter());
+	}
+
+	virtual FString ToString(const double& Value) const override
+	{
+		TSharedPtr<ISequencer>   Sequencer    = WeakSequencer.Pin();
+		UMovieScene*             MovieScene   = Sequencer ? Sequencer->GetFocusedMovieSceneSequence()->GetMovieScene() : nullptr;
+		UMovieSceneMusicalClock* MusicalClock = MovieScene ? MovieScene->GetCustomClockAs<UMovieSceneMusicalClock>() : nullptr;
+		if (!MusicalClock || !MusicalClock->MusicMap)
+		{
+			return FString(TEXT("No music map to determine time"));
+		}
+
+		FrameBasedMusicMap::EBarBeatFormat Format = (Intent == ENumericIntent::Position)
+			? FrameBasedMusicMap::EBarBeatFormat::Position
+			: FrameBasedMusicMap::EBarBeatFormat::Duration;
+		return MusicalClock->MusicMap->MakeBarBeatStringForFrame(Value, Format);
+	}
+
+	virtual TOptional<double> FromString(const FString& InString, const double& InExistingValue) override
+	{
+		TSharedPtr<ISequencer>   Sequencer    = WeakSequencer.Pin();
+		UMovieScene*             MovieScene   = Sequencer ? Sequencer->GetFocusedMovieSceneSequence()->GetMovieScene() : nullptr;
+		UMovieSceneMusicalClock* MusicalClock = MovieScene ? MovieScene->GetCustomClockAs<UMovieSceneMusicalClock>() : nullptr;
+		if (!MusicalClock || !MusicalClock->MusicMap)
+		{
+			return TOptional<double>();
+		}
+		static FBarBeatStringParser StringParser;
+		
+		FrameBasedMusicMap::EBarBeatFormat Format = (Intent == ENumericIntent::Position)
+			? FrameBasedMusicMap::EBarBeatFormat::Position
+			: FrameBasedMusicMap::EBarBeatFormat::Duration;
+
+		TValueOrError<FFrameTime, FExpressionError> Result = StringParser.EvaluateBarBeat(*InString, MusicalClock->MusicMap, Format);
+		if (Result.IsValid())
+		{
+			return Result.GetValue().GetFrame().Value;
+		}
+		return TOptional<double>();
+	}
+};
+
+
+struct FBarsBeatsSequencerNumericTypeInterface : FSequencerNumericTypeInterface
+{
+	FBarsBeatsSequencerNumericTypeInterface(TWeakPtr<ISequencer> WeakSequencer, ENumericIntent Intent)
+		: FSequencerNumericTypeInterface(MakeShared<FBarsBeatsNumericTypeInterface>(WeakSequencer, Intent), Intent)
+	{}
+
+	int32 GetRelevancyScore(const ISequencer& Sequencer, TSharedPtr<IPropertyHandle> Property) const override
+	{
+		// We can only use this if there is a musical clock active
+		UMovieScene*             MovieScene   = Sequencer.GetFocusedMovieSceneSequence()->GetMovieScene();
+		UMovieSceneMusicalClock* MusicalClock = MovieScene ? MovieScene->GetCustomClockAs<UMovieSceneMusicalClock>() : nullptr;
+
+		if (!MusicalClock || !MusicalClock->MusicMap || MusicalClock->MusicMap->IsEmpty())
+		{
+			return MIN_int32;
+		}
+
+		if (Property)
+		{
+			const FString& MetaUIDisplayAsString = Property->GetMetaData(TEXT("UIFrameDisplayAs"));
+
+			// Don't use this for non-musical properties
+			if (MetaUIDisplayAsString.Equals(TEXT("NonMusical"), ESearchCase::IgnoreCase))
+			{
+				return MIN_int32;
+			}
+		}
+
+		const int32 Default = FSequencerNumericTypeInterface::GetRelevancyScore(Sequencer, Property);
+		if (Default != MIN_int32)
+		{
+			// Override the default
+			return Default + 1000;
+		}
+		return MIN_int32;
+	}
+};
+
+void FMusicalModeClockFactoryExtension::OnCreated(TSharedRef<FViewModel> InWeakOwner)
+{
+	WeakSequence = InWeakOwner->CastThisShared<FSequenceModel>();
+}
+
+void FMusicalModeClockFactoryExtension::OnPostInitialize()
+{
+	TViewModelPtr<FSequenceModel> Sequence = WeakSequence.Pin();
+	if (Sequence)
+	{
+		Sequence->GetMovieScene()->EventHandlers.Link(this);
+		ToggleClockExtension();
+	}
+}
+
+void FMusicalModeClockFactoryExtension::OnClockChanged(EUpdateClockSource NewClock, UMovieSceneClock* NewCustomClock)
+{
+	ToggleClockExtension();
+}
+
+void FMusicalModeClockFactoryExtension::ToggleClockExtension()
+{
+	TViewModelPtr<FSequenceModel> Sequence = WeakSequence.Pin();
+	if (Sequence)
+	{
+		UMovieScene*             MovieScene   = Sequence->GetMovieScene();
+		UMovieSceneMusicalClock* MusicalClock = MovieScene ? MovieScene->GetCustomClockAs<UMovieSceneMusicalClock>() : nullptr;
+
+		if (MusicalClock)
+		{
+			Sequence->AddDynamicExtension(FMusicalModeClockExtension::ID);
+			return;
+		}
+		// NOTE: This seems weird to do this here. We are assuming that is tht display format is custom
+		// it must be OUR custom display format... But...
+		// What if the sequence as a custom clock that is NOT a musical clock?
+		// won't we be resetting the display format for them when we don't need to?
+		if (Sequence->GetSequencer()->GetSequencerSettings()->GetTimeDisplayFormat() == EFrameNumberDisplayFormats::Custom)
+		{
+			Sequence->GetSequencer()->GetSequencerSettings()->SetTimeDisplayFormat(EFrameNumberDisplayFormats::Seconds);
+		}
+		Sequence->RemoveDynamicExtension(FMusicalModeClockExtension::ID);
+	}
+}
+
+FMusicalModeClockExtension::~FMusicalModeClockExtension()
+{
+	if (TViewModelPtr<FSequenceModel> Sequence = WeakSequence.Pin())
+	{
+		Sequence->GetSequencer()->RemoveNumericTypeInterface(PositionNumericTimeInterface.ToSharedRef());
+		Sequence->GetSequencer()->RemoveNumericTypeInterface(DurationNumericTimeInterface.ToSharedRef());
+	}
+}
+
+void FMusicalModeClockExtension::OnCreated(TSharedRef<FViewModel> InWeakOwner)
+{
+	WeakSequence = InWeakOwner->CastThisShared<FSequenceModel>();
+	if (TViewModelPtr<FSequenceModel> Sequence = WeakSequence.Pin())
+	{
+		PositionNumericTimeInterface = MakeShared<FBarsBeatsSequencerNumericTypeInterface>(Sequence->GetSequencer(), ENumericIntent::Position);
+		DurationNumericTimeInterface = MakeShared<FBarsBeatsSequencerNumericTypeInterface>(Sequence->GetSequencer(), ENumericIntent::Duration);
+
+		Sequence->GetSequencer()->AddNumericTypeInterface(PositionNumericTimeInterface.ToSharedRef());
+		Sequence->GetSequencer()->AddNumericTypeInterface(DurationNumericTimeInterface.ToSharedRef());
+	}
+}
+
+double FMusicalModeClockExtension::GetDesiredTimeSliderHeight(TSharedPtr<ISequencer> Sequencer) const
+{
+	UMovieScene*             MovieScene   = Sequencer->GetFocusedMovieSceneSequence()->GetMovieScene();
+	UMovieSceneMusicalClock* MusicalClock = MovieScene ? MovieScene->GetCustomClockAs<UMovieSceneMusicalClock>() : nullptr;
+
+	if (MusicalClock && MusicalClock->MusicMap && !MusicalClock->MusicMap->IsEmpty() && MusicalClock->bDisplayFramesAsBarsBeats && Sequencer->GetSequencerSettings()->GetTimeDisplayFormat() == EFrameNumberDisplayFormats::Custom)
+	{
+		return 55.0;
+	}
+
+	return 22.0;
+}
+
+bool FMusicalModeClockExtension::DrawTicks(TSharedPtr<ISequencer> Sequencer, FSlateWindowElementList& OutDrawElements, const TRange<double>& ViewRange, const FSequencerTimeSliderController::FScrubRangeToScreen& RangeToScreen, FSequencerTimeSliderController::FDrawTickArgs& InArgs) const
+{
+	UMovieScene*             MovieScene   = Sequencer->GetFocusedMovieSceneSequence()->GetMovieScene();
+	UMovieSceneMusicalClock* MusicalClock = MovieScene ? MovieScene->GetCustomClockAs<UMovieSceneMusicalClock>() : nullptr;
+
+	if (!MusicalClock || !MusicalClock->MusicMap || MusicalClock->MusicMap->IsEmpty() || !MusicalClock->bDisplayFramesAsBarsBeats || Sequencer->GetSequencerSettings()->GetTimeDisplayFormat() != EFrameNumberDisplayFormats::Custom)
+	{
+		return false;
+	}
+
+	FFrameRate     TickResolution = Sequencer->GetFocusedTickResolution();
+	FFrameRate     DisplayRate    = Sequencer->GetFocusedDisplayRate();
+	FPaintGeometry PaintGeometry = InArgs.AllottedGeometry.ToPaintGeometry();
+	FSlateFontInfo TickFrameFont = FCoreStyle::GetDefaultFontStyle("Regular", 8);
+
+	double PixelsPerSecond = InArgs.AllottedGeometry.Size.X / (ViewRange.GetUpperBoundValue() - ViewRange.GetLowerBoundValue());
+	double MinSecondGap = 20.0f / PixelsPerSecond;
+
+	FFrameTime StartFrameTime = TickResolution.AsFrameTime(ViewRange.GetLowerBoundValue());
+	FFrameTime EndFrameTime = TickResolution.AsFrameTime(ViewRange.GetUpperBoundValue());
+	FFrameTime MinSpacing = TickResolution.AsFrameTime(MinSecondGap);
+
+	TArray<FSubdivisionPoint> Points = MusicalClock->MusicMap->GatherSubdivisonsForTimespan(StartFrameTime, EndFrameTime, MinSpacing, FrameBasedMusicMap::EWholeQuantizationTicks::QuarterNote);
+
+	TArray<FVector2D> LinePoints;
+	LinePoints.SetNumUninitialized(2);
+
+	TSharedRef<FSlateFontMeasure> FontMeasureService = FSlateApplication::Get().GetRenderer()->GetFontMeasureService();
+
+	const bool bAntiAliasLines = false;
+	float LastLabelRightEdge = -10000;
+	for (const FSubdivisionPoint& Point : Points)
+	{
+		float LinePx = RangeToScreen.InputToLocalX(TickResolution.AsSeconds(Point.FrameTime));
+
+		if (Point.MusicalTime.TickInBar == 0)
+		{
+			if (LinePx > (LastLabelRightEdge + 26))
+			{
+				LinePoints[0] = FVector2D(LinePx, InArgs.TickOffset - 25);
+				LinePoints[1] = FVector2D(LinePx, InArgs.TickOffset + InArgs.MajorTickHeight);
+
+				// Space the text between the tick mark but slightly above
+				FVector2D TextOffset(LinePx + 5.f, InArgs.bMirrorLabels ? 1.f : FMath::Abs(InArgs.AllottedGeometry.Size.Y - (30/*InArgs.MajorTickHeight*/ + 9.f)));
+				FSlateDrawElement::MakeText(
+					OutDrawElements,
+					InArgs.StartLayer + 1,
+					InArgs.AllottedGeometry.ToPaintGeometry(InArgs.AllottedGeometry.Size, FSlateLayoutTransform(TextOffset)),
+					Point.PrettyString,
+					TickFrameFont,
+					InArgs.DrawEffects,
+					InArgs.TickColor * 0.65f
+				);
+
+				FVector2D LabelSize = FontMeasureService->Measure(Point.PrettyString, TickFrameFont);
+				LastLabelRightEdge = LinePx + 5.f + LabelSize.Y;
+			}
+			else
+			{
+				LinePoints[0] = FVector2D(LinePx, InArgs.TickOffset);
+				LinePoints[1] = FVector2D(LinePx, InArgs.TickOffset + InArgs.MajorTickHeight);
+			}
+
+			FSlateDrawElement::MakeLines(
+				OutDrawElements,
+				InArgs.StartLayer,
+				PaintGeometry,
+				LinePoints,
+				InArgs.DrawEffects,
+				InArgs.TickColor,
+				bAntiAliasLines
+			);
+		}
+		else
+		{ 
+			const float MinorTickHeight = 2.0f; // : 2.0f;
+
+			LinePoints[0] = FVector2D(LinePx, InArgs.bMirrorLabels ? 0.0f : FMath::Abs(InArgs.AllottedGeometry.Size.Y - MinorTickHeight));
+			LinePoints[1] = FVector2D(LinePx, LinePoints[0].Y + MinorTickHeight);
+
+			FSlateDrawElement::MakeLines(
+				OutDrawElements,
+				InArgs.StartLayer,
+				PaintGeometry,
+				LinePoints,
+				InArgs.DrawEffects,
+				InArgs.TickColor,
+				bAntiAliasLines
+			);
+		}
+	}
+
+	return true;
+}
+
+bool FMusicalModeClockExtension::ShouldShowPlayRateCombo(TSharedPtr<const ISequencer> Sequencer) const
+{
+	UMovieScene*             MovieScene   = Sequencer->GetFocusedMovieSceneSequence()->GetMovieScene();
+	UMovieSceneMusicalClock* MusicalClock = MovieScene ? MovieScene->GetCustomClockAs<UMovieSceneMusicalClock>() : nullptr;
+	return MusicalClock == nullptr;
+}
+} // namespace UE::Sequencer
\ No newline at end of file
diff --git a/Engine/Plugins/MovieScene/MusicalMode/Source/MusicalModeEditor/Private/MusicalModeClockExtension.h b/Engine/Plugins/MovieScene/MusicalMode/Source/MusicalModeEditor/Private/MusicalModeClockExtension.h
new file mode 100644
--- /dev/null
+++ b/Engine/Plugins/MovieScene/MusicalMode/Source/MusicalModeEditor/Private/MusicalModeClockExtension.h
@@ -0,0 +1,59 @@
+// Copyright Epic Games, Inc. All Rights Reserved.
+
+#pragma once
+
+#include "MVVM/ViewModelTypeID.h"
+#include "MVVM/ICastable.h"
+#include "MVVM/Extensions/IClockExtension.h"
+#include "MVVM/Extensions/DynamicExtensionContainer.h"
+#include "MVVM/ViewModelPtr.h"
+#include "EventHandlers/ISequenceDataEventHandler.h"
+
+namespace UE::Sequencer
+{
+
+class FSequenceModel;
+struct FBarsBeatsSequencerNumericTypeInterface;
+
+class FMusicalModeClockFactoryExtension
+	: public IDynamicExtension
+	, private UE::MovieScene::TIntrusiveEventHandler<UE::MovieScene::ISequenceDataEventHandler>
+{
+public:
+	UE_SEQUENCER_DECLARE_VIEW_MODEL_TYPE_ID(FMusicalModeClockFactoryExtension)
+
+	void OnCreated(TSharedRef<FViewModel> InWeakOwner) override;
+	void OnPostInitialize() override;
+	void OnClockChanged(EUpdateClockSource NewClock, UMovieSceneClock* NewCustomClock) override;
+
+	void ToggleClockExtension();
+private:
+	TWeakViewModelPtr<FSequenceModel> WeakSequence;
+};
+
+class FMusicalModeClockExtension
+	: public IDynamicExtension
+	, public IClockExtension
+{
+public:
+	UE_SEQUENCER_DECLARE_VIEW_MODEL_TYPE_ID(FMusicalModeClockExtension)
+
+	using Implements = TImplements<IClockExtension>;
+
+	~FMusicalModeClockExtension();
+
+	void OnCreated(TSharedRef<FViewModel> InWeakOwner) override;
+
+	double GetDesiredTimeSliderHeight(TSharedPtr<ISequencer> Sequencer) const override;
+	bool DrawTicks(TSharedPtr<ISequencer> Sequencer, FSlateWindowElementList& OutDrawElements, const TRange<double>& ViewRange, const FSequencerTimeSliderController::FScrubRangeToScreen& RangeToScreen, FSequencerTimeSliderController::FDrawTickArgs& InArgs) const override;
+
+	bool ShouldShowPlayRateCombo(TSharedPtr<const ISequencer> Sequencer) const override;
+
+private:
+	TWeakViewModelPtr<FSequenceModel> WeakSequence;
+	TSharedPtr<FBarsBeatsSequencerNumericTypeInterface> PositionNumericTimeInterface;
+	TSharedPtr<FBarsBeatsSequencerNumericTypeInterface> DurationNumericTimeInterface;
+};
+
+} // namespace UE::Sequencer
+
diff --git a/Engine/Plugins/MovieScene/MusicalMode/Source/MusicalModeEditor/Private/MusicalModeEditorModule.cpp b/Engine/Plugins/MovieScene/MusicalMode/Source/MusicalModeEditor/Private/MusicalModeEditorModule.cpp
new file mode 100644
--- /dev/null
+++ b/Engine/Plugins/MovieScene/MusicalMode/Source/MusicalModeEditor/Private/MusicalModeEditorModule.cpp
@@ -0,0 +1,291 @@
+// Copyright Epic Games, Inc. All Rights Reserved.
+
+#include "FrameBasedMusicMap.h"
+#include "FrameNumberDisplayFormat.h"
+#include "Framework/MultiBox/MultiBoxBuilder.h"
+#include "IKeyArea.h"
+#include "Modules/ModuleInterface.h"
+#include "Modules/ModuleManager.h"
+#include "MovieScene.h"
+#include "MovieSceneMusicalClock.h"
+#include "MusicalModeClockExtension.h"
+#include "MVVM/Selection/Selection.h"
+#include "MVVM/ViewModels/SequenceModel.h"
+#include "MVVM/ViewModels/SequencerEditorViewModel.h"
+#include "ScopedTransaction.h"
+#include "SequencerSettings.h"
+#include "SequencerToolMenuContext.h"
+#include "SSequencerMusicModeCombo.h"
+#include "ToolMenus.h"
+
+
+UE_SEQUENCER_DEFINE_VIEW_MODEL_TYPE_ID(UE::Sequencer::FMusicalModeClockFactoryExtension)
+UE_SEQUENCER_DEFINE_VIEW_MODEL_TYPE_ID(UE::Sequencer::FMusicalModeClockExtension)
+
+class FMusicalModelEditorModule
+	: public IModuleInterface
+{
+public:
+
+	virtual void StartupModule() override
+	{
+		using namespace UE::Sequencer;
+
+		ExtendSequencerToolbar();
+		ExtendClockSourceMenu();
+		ExtendTimeDisplayFormatMenu();
+		ExtendKeyMenus();
+		ExtendTransportControls();
+
+		FSequenceModel::CreateExtensionsEvent.AddRaw(this, &FMusicalModelEditorModule::ExtendSequenceModel);
+	}
+
+
+	void ExtendSequenceModel(TSharedPtr<UE::Sequencer::FEditorViewModel>, TSharedPtr<UE::Sequencer::FSequenceModel> SequenceModel)
+	{
+		using namespace UE::Sequencer;
+		SequenceModel->AddDynamicExtension(FMusicalModeClockFactoryExtension::ID);
+	}
+
+	void ExtendSequencerToolbar()
+	{
+		FToolMenuOwnerScoped OwnerScoped(this);
+
+		UToolMenu* ToolMenu = UToolMenus::Get()->ExtendMenu("Sequencer.MainToolBar");
+
+		ToolMenu->AddDynamicSection("MusicModeTempo", FNewToolMenuDelegate::CreateLambda([this](UToolMenu* InMenu)
+		{
+			USequencerToolMenuContext* ContextObject = InMenu->FindContext<USequencerToolMenuContext>();
+			TWeakPtr<ISequencer>       WeakSequencer = ContextObject ? ContextObject->WeakSequencer : nullptr;
+			TSharedPtr<ISequencer>     Sequencer     = WeakSequencer.Pin();
+			if (!Sequencer)
+			{
+				return;
+			}
+			
+			FToolMenuEntry MusicModeEntry = FToolMenuEntry::InitWidget(
+			   "MusicMode",
+			   SNew(SSequencerMusicModeCombo, Sequencer)
+				.Visibility_Lambda([=]()
+				{
+					TSharedPtr<ISequencer>   Sequencer    = WeakSequencer.Pin();
+					UMovieScene*             MovieScene   = Sequencer ? Sequencer->GetFocusedMovieSceneSequence()->GetMovieScene() : nullptr;
+					UMovieSceneMusicalClock* MusicalClock = MovieScene ? MovieScene->GetCustomClockAs<UMovieSceneMusicalClock>() : nullptr;
+					return MusicalClock == nullptr ? EVisibility::Collapsed : EVisibility::Visible; 
+				}),
+			   NSLOCTEXT("MusicMode", "MusicMode", "MusicMode")
+			);
+			MusicModeEntry.StyleNameOverride = "SequencerToolbar";
+			FToolMenuSection& Section = InMenu->AddSection("Snapping");
+			Section.AddEntry(MusicModeEntry);
+			Section.AddSeparator(NAME_None);
+		}), FToolMenuInsert("Snapping", EToolMenuInsertType::Before));
+	}
+
+	void ExtendClockSourceMenu()
+	{
+		FToolMenuOwnerScoped OwnerScoped(this);
+
+		UToolMenu* ToolMenu = UToolMenus::Get()->ExtendMenu("Sequencer.PlayRateCombo.ClockSource");
+		ToolMenu->AddDynamicSection("MusicModeTempo", FNewToolMenuDelegate::CreateLambda([this](UToolMenu* InMenu)
+			{
+				USequencerClockSourceMenuContext* ContextObject = InMenu->FindContext<USequencerClockSourceMenuContext>();
+				TWeakPtr<ISequencer>              WeakSequencer = ContextObject ? ContextObject->WeakSequencer : nullptr;
+				TSharedPtr<ISequencer>            Sequencer     = WeakSequencer.Pin();
+				if (!Sequencer)
+				{
+					return;
+				}
+				FToolUIActionChoice SetMusicClockAction(FUIAction(FExecuteAction::CreateLambda([=]()
+					{
+						TSharedPtr<ISequencer> Sequencer = WeakSequencer.Pin();
+						UMovieSceneSequence* MovieSceneSequence = Sequencer ? Sequencer->GetFocusedMovieSceneSequence() : nullptr;
+						UMovieScene* MovieScene = MovieSceneSequence ? MovieSceneSequence->GetMovieScene() : nullptr;
+						if (!MovieScene)
+						{
+							return;
+						}
+						UMovieSceneMusicalClock* MusicalClock = NewObject<UMovieSceneMusicalClock>(MovieScene, NAME_None, RF_Transactional);
+						MusicalClock->CreateSimpleMusicMap(120.0f, { 4, 4 });
+						FScopedTransaction ScopedTransaction(FText(NSLOCTEXT("MusicMode", "EnableMusicMode", "Set Sequence Clock To Musical")));
+						MovieScene->Modify();
+						MovieScene->SetCustomClock(MusicalClock);
+						Sequencer->GetSequencerSettings()->SetFlags(RF_Transactional);
+						Sequencer->GetSequencerSettings()->Modify();
+						Sequencer->GetSequencerSettings()->SetTimeDisplayFormat(EFrameNumberDisplayFormats::Custom);
+					})));
+				FToolMenuEntry SetMusicClock = FToolMenuEntry::InitMenuEntry(
+					"MusicClockWidget",
+					NSLOCTEXT("MusicMode", "MusicClockMenuItem", "Musical Clock"),
+					NSLOCTEXT("MusicMode", "MusicClockMenuItem_ToolTip", "Puts the sequence in 'music mode'"),
+					FSlateIcon(),
+					SetMusicClockAction,
+					EUserInterfaceActionType::RadioButton);
+				InMenu->AddMenuEntry("ClockSourceMusicalClock", SetMusicClock);
+			}));
+	}
+
+	void ExtendKeyMenus()
+	{
+		FToolMenuOwnerScoped OwnerScoped(this);
+
+		UToolMenu* ToolMenu = UToolMenus::Get()->ExtendMenu("Sequencer.KeyContextMenu");
+
+		ToolMenu->AddDynamicSection(NAME_None, FNewToolMenuDelegate::CreateLambda([this](UToolMenu* InMenu)
+		{
+			USequencerToolMenuContext* ContextObject     = InMenu->FindContext<USequencerToolMenuContext>();
+			TWeakPtr<ISequencer>       WeakSequencer     = ContextObject ? ContextObject->WeakSequencer : nullptr;
+			TSharedPtr<ISequencer>     Sequencer         = WeakSequencer.Pin();
+			UMovieScene*               FocusedMovieScene = Sequencer ? Sequencer->GetFocusedMovieSceneSequence()->GetMovieScene() : nullptr;
+			UMovieSceneMusicalClock*   MusicalClock      = FocusedMovieScene ? FocusedMovieScene->GetCustomClockAs<UMovieSceneMusicalClock>() : nullptr;
+
+			if (MusicalClock && MusicalClock->MusicMap && !MusicalClock->MusicMap->IsEmpty() && Sequencer->GetViewModel()->GetSelection()->KeySelection.Num() > 0)
+			{
+				FToolMenuSection& QuantizeDivisionsSection = InMenu->AddSection(TEXT("Quantize"));
+				FToolMenuEntry& SubMenu = QuantizeDivisionsSection.AddSubMenu(
+					TEXT("Divisions"),
+					NSLOCTEXT("MusicMode", "KeyQuantizeTime", "Quantize Key Time"),
+					NSLOCTEXT("MusicMode", "KeyQuantizeTimeTooltip", "Quantize the key time to a musical subdivision"),
+					FNewToolMenuDelegate::CreateLambda([this, WeakSequencer](UToolMenu* InSubMenu)
+						{
+							FName Name = TEXT("Entry");
+							for (TSharedPtr<int32>& Item : FrameBasedMusicMap::QuantizationOptionsSource)
+							{
+								const int32 QunatizationInterval = *Item;
+								Name.SetNumber(QunatizationInterval);
+						
+								FToolMenuEntry MenuEntry = FToolMenuEntry::InitMenuEntry(Name,
+									FrameBasedMusicMap::GetQuantizationName(*Item),
+									FText(),
+									FSlateIcon(),
+									FToolUIActionChoice(FUIAction(FExecuteAction::CreateRaw(this, &FMusicalModelEditorModule::QuantizeSelectedKeyTimes, WeakSequencer, QunatizationInterval))));
+								InSubMenu->AddMenuEntry(NAME_None, MenuEntry);
+							}
+						})
+				);
+			}
+		}), FToolMenuInsert("Keys", EToolMenuInsertType::After));
+	}
+
+	void ExtendTimeDisplayFormatMenu()
+	{
+		UToolMenu* ToolMenu = UToolMenus::Get()->ExtendMenu("Sequencer.TimeDisplayMenu");
+
+		ToolMenu->AddDynamicSection(NAME_None, FNewToolMenuDelegate::CreateLambda([](UToolMenu* InMenu)
+		{
+			USequencerToolMenuContext* ContextObject = InMenu->FindContext<USequencerToolMenuContext>();
+			TWeakPtr<ISequencer>       WeakSequencer = ContextObject ? ContextObject->WeakSequencer : nullptr;
+			TSharedPtr<ISequencer>     Sequencer     = WeakSequencer.Pin();
+			if (!Sequencer)
+			{
+				return;
+			}
+
+			UMovieSceneSequence*     FocusedSequence   = Sequencer->GetFocusedMovieSceneSequence();
+			UMovieScene*             FocusedMovieScene = FocusedSequence ? FocusedSequence->GetMovieScene() : nullptr;
+			UMovieSceneMusicalClock* MusicalClock      = FocusedMovieScene ? FocusedMovieScene->GetCustomClockAs<UMovieSceneMusicalClock>() : nullptr;
+
+			USequencerSettings* Settings = Sequencer->GetSequencerSettings();
+
+			if (MusicalClock)
+			{
+				FToolMenuEntry Entry = FToolMenuEntry::InitMenuEntry(
+					"BarseAndBeats",
+					NSLOCTEXT("MusicMode", "BarsBeatsText", "Bars : Beats"),
+					NSLOCTEXT("MusicMode", "BarsBeatsText", "Show time as Bars : Beats at the current tempo"),
+					FSlateIcon(),
+					FToolUIActionChoice(
+						FUIAction(
+							FExecuteAction::CreateLambda([MusicalClock, Settings]{ MusicalClock->bDisplayFramesAsBarsBeats = true; Settings->SetTimeDisplayFormat(EFrameNumberDisplayFormats::Custom); }),
+							FCanExecuteAction(),
+							FIsActionChecked::CreateLambda([MusicalClock, Settings] { return MusicalClock->bDisplayFramesAsBarsBeats && Settings->GetTimeDisplayFormat() == EFrameNumberDisplayFormats::Custom; })
+						)
+					),
+					EUserInterfaceActionType::RadioButton
+				);
+				InMenu->AddMenuEntry(NAME_None, Entry);
+			}
+		}), FToolMenuInsert("DefaultFormats", EToolMenuInsertType::After));
+	}
+
+	void ExtendTransportControls()
+	{
+		FToolMenuOwnerScoped OwnerScoped(this);
+		UToolMenu* TransportControlsRight = UToolMenus::Get()->ExtendMenu("Sequencer.TransportControls.Right");
+
+		TransportControlsRight->AddDynamicSection("MusicModeTempo", FNewToolMenuDelegate::CreateStatic([](UToolMenu* InMenu)
+		{
+			USequencerToolMenuContext* ContextObject = InMenu->FindContext<USequencerToolMenuContext>();
+			TWeakPtr<ISequencer>       WeakSequencer = ContextObject ? ContextObject->WeakSequencer : nullptr;
+			TSharedPtr<ISequencer>     Sequencer     = WeakSequencer.Pin();
+			if (Sequencer)
+			{
+				TSharedRef<SWidget> Widget = SNew(SSequencerMusicModeTempoEntryBox, WeakSequencer)
+					.UseSpinBoxStyle(true)
+					.AuditionTempo(true)
+					.Visibility_Lambda([WeakSequencer]()
+					{
+						TSharedPtr<ISequencer> Sequencer          = WeakSequencer.Pin();
+						UMovieSceneSequence*   MovieSceneSequence = Sequencer ? Sequencer->GetRootMovieSceneSequence() : nullptr;
+						UMovieScene*           MovieScene         = MovieSceneSequence ? MovieSceneSequence->GetMovieScene() : nullptr;
+						if (!MovieScene || MovieScene->GetCustomClockAs<UMovieSceneMusicalClock>() == nullptr)
+						{
+							return EVisibility::Hidden;
+						}
+						return EVisibility::Visible;
+					});
+				InMenu->AddMenuEntry(NAME_None, FToolMenuEntry::InitWidget("MusicModeTempo", Widget, FText()));
+			}
+		}));
+	}
+
+	virtual void ShutdownModule() override
+	{
+	}
+
+	void QuantizeSelectedKeyTimes(TWeakPtr<ISequencer> WeakSequencer, int32 QuantizationInterval)
+	{
+		using namespace UE::Sequencer;
+
+		TSharedPtr<ISequencer>   Sequencer         = WeakSequencer.Pin();
+		UMovieScene*             FocusedMovieScene = Sequencer ? Sequencer->GetFocusedMovieSceneSequence()->GetMovieScene() : nullptr;
+		UMovieSceneMusicalClock* MusicalClock      = FocusedMovieScene ? FocusedMovieScene->GetCustomClockAs<UMovieSceneMusicalClock>() : nullptr;
+
+		if (!MusicalClock || !MusicalClock->MusicMap)
+		{
+			return;
+		}
+
+		bool bAnythingChanged = false;
+
+		FScopedTransaction SetKeyTimeTransaction(NSLOCTEXT("Sequencer", "SetKeyTime_Transaction", "Set Key Time"));
+		const FKeySelection& KeySelection = Sequencer->GetViewModel()->GetSelection()->KeySelection;
+		for (FKeyHandle Key : KeySelection)
+		{
+			TSharedPtr<FChannelModel> Channel = KeySelection.GetModelForKey(Key);
+			UMovieSceneSection* Section = Channel ? Channel->GetSection() : nullptr;
+			if (Channel && Section)
+			{
+				if (Section->TryModify())
+				{
+					FFrameTime FrameTime = Channel->GetKeyArea()->GetKeyTime(Key);
+					FMusicalTime MusicalTime = MusicalClock->MusicMap->GetMusicalTime(FrameTime);
+					FMusicalTime QuantizedMusicalTime = MusicalClock->MusicMap->Quantize(MusicalTime, QuantizationInterval);
+					FrameTime = MusicalClock->MusicMap->GetFrameTime(QuantizedMusicalTime);
+					Channel->GetKeyArea()->SetKeyTime(Key, FrameTime.GetFrame());
+					bAnythingChanged = true;
+					Section->ExpandToFrame(FrameTime.GetFrame());
+				}
+			}
+		}
+
+		if (bAnythingChanged)
+		{
+			Sequencer->NotifyMovieSceneDataChanged(EMovieSceneDataChangeType::TrackValueChanged);
+		}
+	}
+};
+
+
+IMPLEMENT_MODULE(FMusicalModelEditorModule, MusicalModeEditor);
+
diff --git a/Engine/Plugins/MovieScene/MusicalMode/Source/MusicalModeEditor/Private/SSequencerMusicModeCombo.cpp b/Engine/Plugins/MovieScene/MusicalMode/Source/MusicalModeEditor/Private/SSequencerMusicModeCombo.cpp
new file mode 100644
--- /dev/null
+++ b/Engine/Plugins/MovieScene/MusicalMode/Source/MusicalModeEditor/Private/SSequencerMusicModeCombo.cpp
@@ -0,0 +1,782 @@
+// Copyright Epic Games, Inc. All Rights Reserved.
+
+#include "SSequencerMusicModeCombo.h"
+#include "ISequencer.h"
+#include "Widgets/SOverlay.h"
+#include "Widgets/Layout/SBox.h"
+#include "Widgets/Text/STextBlock.h"
+#include "Widgets/Input/SComboButton.h"
+#include "Widgets/Input/SEditableTextBox.h"
+#include "Widgets/Input/SSpinBox.h"
+#include "Widgets/Images/SImage.h"
+#include "Framework/MultiBox/MultiBoxBuilder.h"
+#include "Misc/CString.h"
+#include "MovieSceneSequence.h"
+#include "MovieScene.h"
+#include "MovieSceneMusicalClock.h"
+#include "MovieSceneCommonHelpers.h"
+#include "Compilation/MovieSceneCompiledDataManager.h"
+#include "ScopedTransaction.h"
+#include "SSequencerMarkerImportDialog.h"
+#include "Widgets/SFrameRateEntryBox.h"
+#include "EditorFontGlyphs.h"
+#include "Evaluation/IMovieSceneCustomClockSource.h"
+#include "SceneOutlinerModule.h"
+#include "SceneOutlinerPublicTypes.h"
+#include "SequencerSettings.h"
+#include "IContentBrowserSingleton.h"
+#include "ContentBrowserModule.h"
+#include "Misc/Timecode.h"
+#include "Misc/AsciiSet.h"
+#include "Styling/ToolBarStyle.h"
+#include "ActorTreeItem.h"
+#include "IMovieSceneRetimingInterface.h"
+
+#define LOCTEXT_NAMESPACE "SSequencerMusicModeCombo"
+
+struct FFrameBasedMusicMapRetimer : UE::MovieScene::IRetimingInterface
+{
+	const UFrameBasedMusicMap* OldMap;
+	const UFrameBasedMusicMap* NewMap;
+
+	FFrameBasedMusicMapRetimer(const UFrameBasedMusicMap* InOldMap, const UFrameBasedMusicMap* InNewMap)
+		: OldMap(InOldMap)
+		, NewMap(InNewMap)
+	{}
+
+	double GetScale() const override
+	{
+		return 1.0;
+	}
+	FFrameTime RemapTime(FFrameTime InTime) const override
+	{
+		return NewMap->RemapFrameTime(InTime, OldMap);
+	}
+
+	TUniquePtr<IRetimingInterface> RecurseInto(UMovieScene* InMovieScene) const override
+	{
+		// Not sure what we do here?
+		return nullptr;
+	}
+
+	void Begin(UMovieScene* InMovieScene) const override
+	{
+
+	}
+	void End(UMovieScene* InMovieScene) const override
+	{
+
+	}
+};
+
+void SSequencerMusicModeCombo::Construct(const FArguments& InArgs, TWeakPtr<ISequencer> InWeakSequencer)
+{
+	WeakSequencer = InWeakSequencer;
+
+	const FToolBarStyle& SequencerToolBarStyle = FAppStyle::Get().GetWidgetStyle<FToolBarStyle>(InArgs._StyleName);
+
+	SetToolTipText(MakeAttributeSP(this, &SSequencerMusicModeCombo::GetToolTipText));
+
+	ChildSlot
+	.VAlign(VAlign_Fill)
+	[
+		SNew(SComboButton)
+		.ContentPadding(FMargin(2.f, 1.0f))
+		.VAlign(VAlign_Fill)
+		.ComboButtonStyle(&SequencerToolBarStyle.ComboButtonStyle)
+		.OnGetMenuContent(this, &SSequencerMusicModeCombo::OnCreateMenu)
+		.ButtonContent()
+		[
+			SNew(SHorizontalBox)
+
+			+ SHorizontalBox::Slot()
+			.HAlign(HAlign_Center)
+			.VAlign(VAlign_Center)
+			.AutoWidth()
+			[
+				SNew(STextBlock)
+				.Text(this, &SSequencerMusicModeCombo::GetTempoTimeSignatureText)
+				.TextStyle(&SequencerToolBarStyle.LabelStyle)
+			]
+		]
+	];
+}
+
+UMovieSceneMusicalClock* SSequencerMusicModeCombo::GetMusicalClock() const
+{
+	TSharedPtr<ISequencer>   Sequencer         = WeakSequencer.Pin();
+	UMovieSceneSequence*     FocusedSequence   = Sequencer.IsValid() ? Sequencer->GetFocusedMovieSceneSequence() : nullptr;
+	UMovieScene*             FocusedMovieScene = FocusedSequence ? FocusedSequence->GetMovieScene() : nullptr;
+
+	return FocusedMovieScene ? FocusedMovieScene->GetCustomClockAs<UMovieSceneMusicalClock>() : nullptr;
+}
+
+bool SSequencerMusicModeCombo::GetIsSequenceReadOnly() const
+{
+	TSharedPtr<ISequencer> Sequencer = WeakSequencer.Pin();
+	return !Sequencer.IsValid() || Sequencer->IsReadOnly();
+}
+
+TSharedRef<SWidget> SSequencerMusicModeCombo::OnCreateMenu()
+{
+	FMenuBuilder MenuBuilder(true, nullptr);
+
+	UMovieSceneMusicalClock* MusicalClock = GetMusicalClock();
+
+	PopulateSimpleMusicMapProperties(MenuBuilder, MusicalClock ? MusicalClock->MusicMap : nullptr);
+	PopulateSimpleMusicMapSection(MenuBuilder, MusicalClock ? MusicalClock->MusicMap : nullptr);
+	PopulateImportMusicMapSection(MenuBuilder);
+
+	if (MusicalClock)
+	{
+		MenuBuilder.AddMenuEntry(
+			LOCTEXT("CearMusicMap", "Remove Music Map"),
+			LOCTEXT("ClearMusicMap_tooltip", "Removes any and all tempo and time signature information."),
+			FSlateIcon(),
+			FUIAction(
+				FExecuteAction::CreateSP(this, &SSequencerMusicModeCombo::OnRemoveAllMusicMaps),
+				FCanExecuteAction::CreateLambda([this] { return !GetIsSequenceReadOnly(); })
+			)
+		);
+	}
+
+	return MenuBuilder.MakeWidget();
+}
+
+void SSequencerMusicModeCombo::PopulateSimpleMusicMapProperties(FMenuBuilder& MenuBuilder, const UFrameBasedMusicMap* MusicMap)
+{
+	if (!MusicMap || !MusicMap->IsSimple())
+	{
+		return;
+	}
+
+	MenuBuilder.BeginSection(NAME_None, LOCTEXT("SimpleMapSectionOptions", "Options"));
+
+	MenuBuilder.AddWidget(
+		SNew(SBox)
+		.HAlign(HAlign_Right)
+		.MaxDesiredWidth(100.f)
+		[
+			SNew(SSequencerMusicModeTempoEntryBox, WeakSequencer)
+			.IsEnabled(this, &SSequencerMusicModeCombo::SimpleMapEntryEnabled)
+		],
+		LOCTEXT("MusicModeSettings_Tempo", "Tempo"),
+		/*InNoIndent*/ true, /*bInSearchable*/ true,
+		LOCTEXT("MusicModeSettings_Tempo_ToolTip", "This is the \"natural tempo\" this sequence will play at when not being driven by an external musical clock source.")
+	);
+	MenuBuilder.AddWidget(
+		SNew(SBox)
+		.HAlign(HAlign_Right)
+		.MaxDesiredWidth(100.f)
+		[
+			SNew(SEditableTextBox)
+				.Text_Lambda([this]()
+					{
+						UMovieSceneMusicalClock* MusicalClock = GetMusicalClock();
+						if (MusicalClock && MusicalClock->MusicMap)
+						{
+							return FText::Format(LOCTEXT("TempoFormat", "{0}/{1}"),
+								FText::AsNumber(MusicalClock->MusicMap->GetInitialTimeSignature().Numerator),
+								FText::AsNumber(MusicalClock->MusicMap->GetInitialTimeSignature().Denominator));
+						}
+						return FText();
+					})
+				.OnTextCommitted(this, &SSequencerMusicModeCombo::OnTimeSignatureTextCommitted)
+				.OnVerifyTextChanged(this, &SSequencerMusicModeCombo::ValidateTimeSignatureString)
+				.IsEnabled(this, &SSequencerMusicModeCombo::SimpleMapEntryEnabled)
+				.SelectAllTextWhenFocused(true)
+		],
+		LOCTEXT("MusicModeSettings_TimeSignature", "Time Signature"),
+		/*InNoIndent*/ true, /*bInSearchable*/ true,
+		LOCTEXT("MusicModeSettings_TimeSignature_ToolTip", "The \"natural time signature\" this sequence is authored for.")
+	);
+
+	MenuBuilder.EndSection();
+	MenuBuilder.AddSeparator();
+}
+
+void SSequencerMusicModeCombo::PopulateSimpleMusicMapSection(FMenuBuilder& MenuBuilder, const UFrameBasedMusicMap* MusicMap)
+{
+	if (!MusicMap)
+	{
+		MenuBuilder.AddMenuEntry(
+			LOCTEXT("MakeSimpleMusicMap", "Add Simple Music Map"),
+			LOCTEXT("MakeSimpleMusicMap_Description", "Creates a simple, one tempo, one time signature music map for the sequence."),
+			FSlateIcon(),
+			FUIAction( FExecuteAction::CreateSP(this, &SSequencerMusicModeCombo::OnCreateSimpleMusicMap),
+				       FCanExecuteAction::CreateLambda([this] { return !GetIsSequenceReadOnly(); }))
+			);
+	}
+	else if (!MusicMap->IsSimple())
+	{
+		MenuBuilder.AddMenuEntry(
+			LOCTEXT("ClearMusicMapAndMakeSimple", "Clear Imported Music Map And Make Simple"),
+			LOCTEXT("ClearMusicMapAndMakeSimple_Description", "Clears all imported tempo and time signature changes and creates a simple, one tempo, one time signature music map for the sequence."),
+			FSlateIcon(),
+			FUIAction(FExecuteAction::CreateSP(this, &SSequencerMusicModeCombo::OnCreateSimpleMusicMap),
+				FCanExecuteAction::CreateLambda([this] { return !GetIsSequenceReadOnly(); }))
+		);
+	}
+}
+
+void SSequencerMusicModeCombo::PopulateImportMusicMapSection(FMenuBuilder& MenuBuilder)
+{
+	auto MusicMapAssetsMenu = [this](FMenuBuilder& ActorMenuBuilder)
+		{
+			FAssetPickerConfig AssetPickerConfig;
+			AssetPickerConfig.Filter.ClassPaths.Add(UMusicMapSource::StaticClass()->GetClassPathName());
+			AssetPickerConfig.bAllowNullSelection = false;
+			AssetPickerConfig.Filter.bRecursiveClasses = true;
+			AssetPickerConfig.OnAssetSelected = FOnAssetSelected::CreateLambda([this](const FAssetData& In) { this->OnImportMusicMap(In.GetAsset()); });
+			AssetPickerConfig.InitialAssetViewType = EAssetViewType::List;
+			AssetPickerConfig.bAllowDragging = false;
+
+			FContentBrowserModule& ContentBrowserModule = FModuleManager::Get().LoadModuleChecked<FContentBrowserModule>(TEXT("ContentBrowser"));
+
+			ActorMenuBuilder.AddWidget(
+				SNew(SBox)
+				.IsEnabled(this, &SSequencerMusicModeCombo::GetIsSequenceReadOnly)
+				.WidthOverride(300)
+				.HeightOverride(300)
+				[
+					ContentBrowserModule.Get().CreateAssetPicker(AssetPickerConfig)
+				],
+				FText(),
+				true /*bNoIndent*/
+			);
+		};
+
+	MenuBuilder.AddSubMenu(
+		LOCTEXT("ImportMusicMap", "Import Music Map..."),
+		LOCTEXT("ImportMusicMap_tooltip", "Imports tempo and time signature changes from the selected asset. NOTE: This will remove any existing tempo and time signature changes!"),
+		FNewMenuDelegate::CreateLambda(MusicMapAssetsMenu)
+	);
+}
+
+void SSequencerMusicModeCombo::PopulateMusicalClockSourceMenu(FMenuBuilder& MenuBuilder)
+{
+#if 0
+    THIS IS THE CRIBBED CODE FOR SEARCHING FOR OBJECTS WITH A SEQUENCER CLOCK. IT WILL GET
+    RESURRECTED WHEN WE GET TO THE POINT WHEN WE ARE LOOKING FOR A MUSICAL SEQUENCER CLOCK.
+	auto ActorClockSourceMenu = [this](FMenuBuilder& ActorMenuBuilder)
+	{
+		auto IsActorValid = [](const AActor* Actor)
+		{
+			if (!Actor)
+			{
+				return false;
+			}
+
+			if (Actor->GetClass()->ImplementsInterface(UMovieSceneCustomClockSource::StaticClass()))
+			{
+				return true;
+			}
+
+			TArray<UActorComponent*> ComponentsWithClockSource = Actor->GetComponentsByInterface(UMovieSceneCustomClockSource::StaticClass());
+			if (ComponentsWithClockSource.Num() > 0)
+			{
+				return true;
+			}
+
+			return false;
+		};
+
+		// Set up a menu entry to assign an actor to the object binding node
+		FSceneOutlinerInitializationOptions InitOptions;
+
+		// We hide the header row to keep the UI compact.
+		InitOptions.bShowHeaderRow = false;
+		InitOptions.bShowSearchBox = true;
+		InitOptions.bShowCreateNewFolder = false;
+		InitOptions.bFocusSearchBoxWhenOpened = true;
+		// Only want the actor label column
+		InitOptions.ColumnMap.Add(FSceneOutlinerBuiltInColumnTypes::Label(), FSceneOutlinerColumnInfo(ESceneOutlinerColumnVisibility::Visible, 0));
+
+		// Only display actors that are not possessed already
+		InitOptions.Filters->AddFilterPredicate<FActorTreeItem>(FActorTreeItem::FFilterPredicate::CreateLambda( IsActorValid ) );
+
+		TSharedPtr<ISequencer> SequencerPtr = WeakSequencer.Pin();
+		const float WidthOverride = SequencerPtr.IsValid() ? SequencerPtr->GetSequencerSettings()->GetAssetBrowserWidth() : 500.f;
+		const float HeightOverride = SequencerPtr.IsValid() ? SequencerPtr->GetSequencerSettings()->GetAssetBrowserHeight() : 400.f;
+
+		// actor selector to allow the user to choose an actor
+		FSceneOutlinerModule& SceneOutlinerModule = FModuleManager::LoadModuleChecked<FSceneOutlinerModule>("SceneOutliner");
+		ActorMenuBuilder.AddWidget(
+			SNew(SBox)
+			.WidthOverride(WidthOverride)
+			.HeightOverride(HeightOverride)
+			[
+				SceneOutlinerModule.CreateActorPicker(InitOptions, FOnActorPicked::CreateLambda([this](AActor* In){ this->SetCustomClockSource(In); }))
+			],
+			FText(),
+			true /*bNoIndent*/
+		);
+	};
+
+	auto AssetClockSourceMenu = [this](FMenuBuilder& ActorMenuBuilder)
+	{
+		FAssetPickerConfig AssetPickerConfig;
+		AssetPickerConfig.Filter.ClassPaths.Add(UMovieSceneCustomClockSource::StaticClass()->GetClassPathName());
+		AssetPickerConfig.bAllowNullSelection = false;
+		AssetPickerConfig.Filter.bRecursiveClasses = true;
+		AssetPickerConfig.OnAssetSelected = FOnAssetSelected::CreateLambda([this](const FAssetData& In){ this->SetCustomClockSource(In.GetAsset()); });
+		AssetPickerConfig.InitialAssetViewType = EAssetViewType::List;
+		AssetPickerConfig.bAllowDragging = false;
+
+		FContentBrowserModule& ContentBrowserModule = FModuleManager::Get().LoadModuleChecked<FContentBrowserModule>(TEXT("ContentBrowser"));
+
+		ActorMenuBuilder.AddWidget(
+			SNew(SBox)
+			.WidthOverride(300)
+			.HeightOverride(300)
+			[
+				ContentBrowserModule.Get().CreateAssetPicker(AssetPickerConfig)
+			],
+			FText(),
+			true /*bNoIndent*/
+		);
+	};
+
+	MenuBuilder.AddSubMenu(LOCTEXT("ActorClockSource", "Actors"), FText(), FNewMenuDelegate::CreateLambda(ActorClockSourceMenu));
+	MenuBuilder.AddSubMenu(LOCTEXT("AssetClockSource", "Assets"), FText(), FNewMenuDelegate::CreateLambda(AssetClockSourceMenu));
+#endif
+}
+
+FText SSequencerMusicModeCombo::GetTempoTimeSignatureText() const
+{
+	TSharedPtr<ISequencer>   Sequencer    = WeakSequencer.Pin();
+	UMovieSceneMusicalClock* MusicalClock = GetMusicalClock();
+
+	if (MusicalClock && MusicalClock->MusicMap)
+	{
+		FNumberFormattingOptions Formatting;
+		Formatting.MaximumFractionalDigits = 3;
+		if (MusicalClock->MusicMap->IsSimple())
+		{
+			float Tempo = MusicalClock->MusicMap->GetInitialTempo();
+			FFrameBasedTimeSignature TimeSignature = MusicalClock->MusicMap->GetInitialTimeSignature();
+			return FText::Format(LOCTEXT("TempoTimeSigDisplay_Simple", "{0} bpm {1}/{2}"),
+				FText::AsNumber(Tempo, &Formatting),
+				FText::AsNumber(TimeSignature.Numerator, &Formatting),
+				FText::AsNumber(TimeSignature.Denominator, &Formatting));
+		}
+		else
+		{
+			FQualifiedFrameTime CursorTime = Sequencer->GetLocalTime();
+			float Tempo = MusicalClock->MusicMap->GetTempoForFrameTime(CursorTime.Time);
+			FFrameBasedTimeSignature TimeSignature = MusicalClock->MusicMap->GetTimeSignatureForFrameTime(CursorTime.Time);
+			return FText::Format(LOCTEXT("TempoTimeSigDisplay_Complex", "Current: {0} bpm {1}/{2}"),
+				FText::AsNumber(Tempo, &Formatting),
+				FText::AsNumber(TimeSignature.Numerator, &Formatting),
+				FText::AsNumber(TimeSignature.Denominator, &Formatting));
+		}
+	}
+
+	return LOCTEXT("MusicModeNotEnabled", "Music Mode...");
+}
+
+FText SSequencerMusicModeCombo::GetToolTipText() const
+{
+	UMovieSceneMusicalClock* MusicalClock = GetMusicalClock();
+	if (MusicalClock && MusicalClock->MusicMap && !MusicalClock->MusicMap->IsSimple())
+	{
+		return FText::Format(LOCTEXT("MusicModeMenuToolTip", "{0} tempos & {1} time signatures"), MusicalClock->MusicMap->TempoMap.Num(), MusicalClock->MusicMap->BarMap.Num());
+	}
+	return FText();
+}
+
+void SSequencerMusicModeCombo::OnRemoveAllMusicMaps()
+{
+	TSharedPtr<ISequencer>   Sequencer       = WeakSequencer.Pin();
+	UMovieSceneSequence*     FocusedSequence = Sequencer.IsValid() ? Sequencer->GetFocusedMovieSceneSequence() : nullptr;
+
+	if (FocusedSequence)
+	{
+		FScopedTransaction ScopedTransaction(FText(LOCTEXT("SetMusicMode_Off", "Turning music mode off for sequence")));
+		FocusedSequence->GetMovieScene()->Modify();
+		FocusedSequence->GetMovieScene()->SetClockSource(EUpdateClockSource::Audio);
+
+		if (Sequencer->GetSequencerSettings()->GetTimeDisplayFormat() == EFrameNumberDisplayFormats::Custom)
+		{
+			Sequencer->GetSequencerSettings()->SetTimeDisplayFormat(EFrameNumberDisplayFormats::Frames);
+		}
+	}
+}
+
+void SSequencerMusicModeCombo::OnCreateSimpleMusicMap()
+{
+	TSharedPtr<ISequencer>   Sequencer         = WeakSequencer.Pin();
+	UMovieSceneSequence*     FocusedSequence   = Sequencer.IsValid() ? Sequencer->GetFocusedMovieSceneSequence() : nullptr;
+	UMovieScene*             FocusedMovieScene = FocusedSequence ? FocusedSequence->GetMovieScene() : nullptr;
+	
+	if (!FocusedMovieScene)
+	{
+		return;
+	}
+	
+	UMovieSceneMusicalClock* MusicalClock = FocusedMovieScene->GetCustomClockAs<UMovieSceneMusicalClock>();
+
+	FScopedTransaction ScopedTransaction(FText(LOCTEXT("SetMusicMode_On", "Turning music mode on for sequence")));
+	if (!MusicalClock)
+	{
+		FocusedMovieScene->Modify();
+		MusicalClock = NewObject<UMovieSceneMusicalClock>(FocusedMovieScene, NAME_None, RF_Transactional);
+	}
+
+	if (MusicalClock->MusicMap && !MusicalClock->MusicMap->IsSimple())
+	{
+		// WARN THE USER HERE! Make sure they want to blow away the existing music map!
+	}
+	
+	MusicalClock->Modify();
+
+	UFrameBasedMusicMap* NewSimpleMusicMap = NewObject<UFrameBasedMusicMap>(MusicalClock, NAME_None, RF_Transactional);
+	NewSimpleMusicMap->SetFrameResolution(MusicalClock->GetTypedOuter<UMovieScene>()->GetTickResolution());
+	NewSimpleMusicMap->Init(120.0f, { 4, 4 });
+	if (MusicalClock->MusicMap && !MusicalClock->MusicMap->IsEmpty())
+	{
+		// This will remap from the old to the new map
+		UE::MovieScene::TimeHelpers::MigrateFrameTimes(FFrameBasedMusicMapRetimer(MusicalClock->MusicMap, NewSimpleMusicMap), FocusedMovieScene);
+		MusicalClock->MusicMap = NewSimpleMusicMap;
+
+	}
+	else
+	{
+		MusicalClock->MusicMap = NewSimpleMusicMap;
+	}
+
+	Sequencer->GetSequencerSettings()->SetTimeDisplayFormat(EFrameNumberDisplayFormats::Custom);
+	MusicalClock->bDisplayFramesAsBarsBeats = true;
+	FocusedMovieScene->SetCustomClock(MusicalClock);
+}
+
+void SSequencerMusicModeCombo::OnImportMusicMap(UObject* Object)
+{
+	TSharedPtr<ISequencer>   Sequencer         = WeakSequencer.Pin();
+	UMovieSceneSequence*     FocusedSequence   = Sequencer.IsValid() ? Sequencer->GetFocusedMovieSceneSequence() : nullptr;
+	UMovieScene*             FocusedMovieScene = FocusedSequence ? FocusedSequence->GetMovieScene() : nullptr;
+	
+	if (!FocusedMovieScene)
+	{
+		return;
+	}
+	
+	UMovieSceneMusicalClock* MusicalClock = FocusedMovieScene->GetCustomClockAs<UMovieSceneMusicalClock>();
+	
+	FScopedTransaction ScopedTransaction(LOCTEXT("RemapSequence", "Remap Sequence to a new music map"));
+	if (!MusicalClock)
+	{
+		FocusedMovieScene->Modify();
+		MusicalClock = NewObject<UMovieSceneMusicalClock>(FocusedMovieScene, NAME_None, RF_Transactional);
+		FocusedMovieScene->SetCustomClock(MusicalClock);
+	}
+
+
+	IMusicMapSource* MapSource = Cast<IMusicMapSource>(Object);
+	check(MapSource);
+
+	UFrameBasedMusicMap* NewMusicMap = NewObject<UFrameBasedMusicMap>(MusicalClock);
+	NewMusicMap->SetFrameResolution(FocusedMovieScene->GetTickResolution());
+	MapSource->CreateFrameBasedMusicMap(NewMusicMap);
+
+	FMarkerProviderResults Markers = MapSource->GatherMarkers(NewMusicMap);
+
+	MusicalClock->Modify();
+
+	// This will remap from the old to the new map
+	UE::MovieScene::TimeHelpers::MigrateFrameTimes(FFrameBasedMusicMapRetimer(MusicalClock->MusicMap, NewMusicMap), FocusedMovieScene);
+	MusicalClock->MusicMap = NewMusicMap;
+
+	if (!Markers.IsEmpty())
+	{
+		PromptForMarkerImport(Markers);
+	}
+}
+
+void SSequencerMusicModeCombo::OnTimeSignatureTextCommitted(const FText& NewValue, ETextCommit::Type)
+{
+	UMovieSceneMusicalClock* MusicalClock = GetMusicalClock();
+
+	if (MusicalClock && MusicalClock->MusicMap)
+	{
+		FString NumeratorString;
+		FString DenominatorString;
+		if (ensure(NewValue.ToString().Split("/", &NumeratorString, &DenominatorString)))
+		{
+			FFrameBasedTimeSignature OriginalTimeSignature = MusicalClock->MusicMap->GetInitialTimeSignature();
+			FFrameBasedTimeSignature NewTimeSignature;
+			NewTimeSignature.Numerator = FCString::Atoi(*NumeratorString);
+			NewTimeSignature.Denominator = FCString::Atoi(*DenominatorString);
+			if (OriginalTimeSignature != NewTimeSignature)
+			{
+				MusicalClock->Modify();
+				MusicalClock->MusicMap->Modify();
+				MusicalClock->MusicMap->SetInitialTimeSignature(NewTimeSignature);
+			}
+		}
+	}
+}
+
+bool SSequencerMusicModeCombo::ValidateTimeSignatureString(const FText& InText, FText& OutErrorMessage) const
+{
+	FString NumeratorString;
+	FString DenominatorString;
+	bool Success = InText.ToString().Split("/", &NumeratorString, &DenominatorString);
+	if (Success)
+	{
+		int32 NewNumerator = FCString::Atoi(*NumeratorString);
+		bool Valid = NewNumerator > 0 && NewNumerator < 100;
+		if (!Valid)
+		{
+			OutErrorMessage = LOCTEXT("MusicModeSettings_Error_Numerator", "Time signature numerator must be from 1 to 99.");
+			return false;
+		}
+		int32 NewDenominator = FCString::Atoi(*DenominatorString);
+		Valid = NewDenominator > 0 && NewDenominator < 100 && FMath::IsPowerOfTwo(NewDenominator);
+		if (!Valid)
+		{
+			OutErrorMessage = LOCTEXT("MusicModeSettings_Error_Denominator", "Time signature denominator must be from 1 to 64 and be a power of 2.");
+			return false;
+		}
+		return true;
+	}
+	OutErrorMessage = LOCTEXT("MusicModeSettings_Error_TimeSignature", "No '/' detected. Time signature must be in the form '4/4', '7/8', etc.");
+	return false;
+}
+
+bool SSequencerMusicModeCombo::SimpleMapEntryEnabled() const
+{
+	if (!GetIsSequenceReadOnly())
+	{
+		UMovieSceneMusicalClock* MusicalClock = GetMusicalClock();
+		return MusicalClock && MusicalClock->MusicMap && MusicalClock->MusicMap->IsSimple();
+	}
+	return false;
+}
+
+
+bool SSequencerMusicModeCombo::PromptForMarkerImport(const FMarkerProviderResults& Markers)
+{
+	TSharedPtr<ISequencer> Sequencer         = WeakSequencer.Pin();
+	UMovieSceneSequence*   FocusedSequence   = Sequencer.IsValid() ? Sequencer->GetFocusedMovieSceneSequence() : nullptr;
+	UMovieScene*           FocusedMovieScene = FocusedSequence ? FocusedSequence->GetMovieScene() : nullptr;
+
+	if (!FocusedMovieScene || Markers.IsEmpty())
+	{
+		return false;
+	}
+
+	FSequencerMarkerImportConfig Results = SSequencerMarkerImportDialog::GetMarkerImportOptions(Markers.Channels[0].Markers);
+	if (!Results.bImport)
+	{
+		return false;
+	}
+
+	if (Results.bClearPreviouslyImported)
+	{
+		const TArray<FMovieSceneMarkedFrame>& ExistingMarkedFrames = FocusedMovieScene->GetMarkedFrames();
+		for (int32 Index = ExistingMarkedFrames.Num() - 1; Index >= 0; --Index)
+		{
+			if (ExistingMarkedFrames[Index].Comment.Compare(TEXT("From Imported Music Map")) == 0)
+			{
+				FocusedMovieScene->DeleteMarkedFrame(Index);
+			}
+		}
+	}
+
+	int32 ImportCount = 0;
+	FocusedMovieScene->Modify();
+	for (const FMarkerProviderEntry& Entry : Markers.Channels[0].Markers)
+	{
+		const TArray<FMovieSceneMarkedFrame>& ExistingMarkedFrames = FocusedMovieScene->GetMarkedFrames();
+		bool bAddIt = true;
+		if (!Results.bClearPreviouslyImported)
+		{
+			for (int32 Index = 0; Index < ExistingMarkedFrames.Num(); ++Index)
+			{
+				if (ExistingMarkedFrames[Index].Label == Entry.Label && ExistingMarkedFrames[Index].Comment.Compare(TEXT("From Imported Music Map")) == 0)
+				{
+					// Label already exists. 
+					// If it is already in the correct place move on...
+					if (ExistingMarkedFrames[Index].FrameNumber == Entry.FrameNumber)
+					{
+						bAddIt = false;
+						break;
+					}
+					// Otherwise delete the old one...
+					FocusedMovieScene->DeleteMarkedFrame(Index);
+					break;
+				}
+			}
+		}
+		if (!bAddIt)
+		{
+			continue;
+		}
+
+		FMovieSceneMarkedFrame MarkedFrame(Entry.FrameNumber);
+		MarkedFrame.Label = Entry.Label;
+		MarkedFrame.Comment = TEXT("From Imported Music Map");
+		FocusedMovieScene->AddMarkedFrame(MarkedFrame);
+		++ImportCount;
+	}
+	return ImportCount > 0;
+}
+
+
+void SSequencerMusicModeTempoEntryBox::Construct(const FArguments& InArgs, TWeakPtr<ISequencer> InWeakSequencer)
+{
+	WeakSequencer = InWeakSequencer;
+
+	struct FNumericInterfaceBPM : INumericTypeInterface<float>
+	{
+		int32 GetMinFractionalDigits() const override { return 0; }
+		int32 GetMaxFractionalDigits() const override { return 3; }
+
+		void SetMinFractionalDigits(const TAttribute<TOptional<int32>>& NewValue) override {}
+		void SetMaxFractionalDigits(const TAttribute<TOptional<int32>>& NewValue) override {}
+
+		FString ToString(const float& Value) const override
+		{
+			return FText::Format(LOCTEXT("MusicModeSettings_Bpm", "{0} bpm"), Value).ToString();
+		}
+		TOptional<float> FromString(const FString& InString, const float& ExistingValue)
+		{
+			float Value = FCString::Atof(*InString);
+			if (Value > 0.f && Value <= 999.f)
+			{
+				return Value;
+			}
+			return TOptional<float>();
+		}
+		bool IsCharacterValid(TCHAR InChar) const override
+		{
+			static const FAsciiSet ValidChars = TEXT("0123456789 bpmBPM/");
+			return ValidChars.Contains(InChar);
+		}
+	};
+
+	bAuditionTempo = InArgs._AuditionTempo;
+
+	if (InArgs._UseSpinBoxStyle)
+	{
+		ChildSlot
+		[
+			SNew(SSpinBox<float>)
+			.Style(&FAppStyle::GetWidgetStyle<FSpinBoxStyle>("Sequencer.HyperlinkSpinBox"))
+			.TypeInterface(MakeShared<FNumericInterfaceBPM>())
+			.OnValueCommitted_Lambda([this](float Value, ETextCommit::Type) { this->SetValue(Value); })
+			.OnValueChanged(this, &SSequencerMusicModeTempoEntryBox::SetValue)
+			.MinValue(1.0)
+			.MaxValue(999.0)
+			.Value(this, &SSequencerMusicModeTempoEntryBox::GetValue)
+			.ToolTipText(LOCTEXT("AuditionTempoSpinBoxToolTip", "Right Click For Authored Tempo"))
+		];
+	}
+	else
+	{
+		ChildSlot
+		[
+			SNew(SSpinBox<float>)
+			//.Style(&FAppStyle::GetWidgetStyle<FSpinBoxStyle>("Sequencer.HyperlinkSpinBox"))
+			.TypeInterface(MakeShared<FNumericInterfaceBPM>())
+			.OnValueCommitted_Lambda([this](float Value, ETextCommit::Type) { this->SetValue(Value); })
+			.OnValueChanged(this, &SSequencerMusicModeTempoEntryBox::SetValue)
+			.MinValue(1.0)
+			.MaxValue(999.0)
+			.Value(this, &SSequencerMusicModeTempoEntryBox::GetValue)
+		];
+	}
+}
+
+FReply SSequencerMusicModeTempoEntryBox::OnMouseButtonDown(const FGeometry& MyGeometry, const FPointerEvent& MouseEvent)
+{
+	if (!bAuditionTempo || MouseEvent.GetEffectingButton() != EKeys::RightMouseButton)
+	{
+		return SCompoundWidget::OnMouseButtonDown(MyGeometry, MouseEvent);
+	}
+	else
+	{
+		TSharedPtr<ISequencer>   Sequencer         = WeakSequencer.Pin();
+		UMovieScene*             FocusedMovieScene = Sequencer.IsValid() ? Sequencer->GetFocusedMovieSceneSequence()->GetMovieScene() : nullptr;
+		UMovieSceneMusicalClock* MusicalClock      = FocusedMovieScene ? FocusedMovieScene->GetCustomClockAs<UMovieSceneMusicalClock>() : nullptr;
+		if (MusicalClock)
+		{
+			MusicalClock->SetAuditionTempo(MusicalClock->MusicMap->GetInitialTempo());
+			return FReply::Handled();
+		}
+		return FReply::Unhandled();
+	}
+}
+
+void SSequencerMusicModeTempoEntryBox::OnTempoTextCommitted(const FText& NewValue, ETextCommit::Type)
+{
+	float NewTempo = FCString::Atof(*(NewValue.ToString()));
+	SetValue(NewTempo);
+}
+
+float SSequencerMusicModeTempoEntryBox::GetValue() const
+{
+	TSharedPtr<ISequencer>   Sequencer         = WeakSequencer.Pin();
+	UMovieScene*             FocusedMovieScene = Sequencer.IsValid() ? Sequencer->GetFocusedMovieSceneSequence()->GetMovieScene() : nullptr;
+	UMovieSceneMusicalClock* MusicalClock      = FocusedMovieScene ? FocusedMovieScene->GetCustomClockAs<UMovieSceneMusicalClock>() : nullptr;
+
+	if (!MusicalClock)
+	{
+		return 0.0;
+	}
+	else if (bAuditionTempo)
+	{
+		return MusicalClock->GetAuditionTempo();
+	}
+	else if (MusicalClock->MusicMap)
+	{
+		return MusicalClock->MusicMap->TempoMap[0].Bpm();
+	}
+	return 0.0;
+}
+
+void SSequencerMusicModeTempoEntryBox::SetValue(float NewTempo)
+{
+	if (NewTempo <= 0.0f || NewTempo > 999.0f)
+	{
+		return;
+	}
+
+	TSharedPtr<ISequencer>   Sequencer         = WeakSequencer.Pin();
+	UMovieScene*             FocusedMovieScene = Sequencer.IsValid() ? Sequencer->GetFocusedMovieSceneSequence()->GetMovieScene() : nullptr;
+	UMovieSceneMusicalClock* MusicalClock      = FocusedMovieScene ? FocusedMovieScene->GetCustomClockAs<UMovieSceneMusicalClock>() : nullptr;
+
+	if (!MusicalClock)
+	{
+		return;
+	}
+
+	if (bAuditionTempo)
+	{
+		if (MusicalClock->GetAuditionTempo() != NewTempo)
+		{
+			FScopedTransaction ScopedTransaction(FText(LOCTEXT("SetAuditionTempo", "Set Audition Tempo")));
+
+			MusicalClock->Modify();
+			MusicalClock->SetAuditionTempo(NewTempo);
+		}
+	}
+	else if (!MusicalClock->MusicMap || MusicalClock->MusicMap->IsEmpty() || MusicalClock->MusicMap->TempoMap[0].Bpm() != NewTempo)
+	{
+		FScopedTransaction ScopedTransaction(FText(LOCTEXT("SetMusicModeTempo", "Set tempo")));
+
+		MusicalClock->Modify();
+
+		UFrameBasedMusicMap* NewMusicMap = NewObject<UFrameBasedMusicMap>(MusicalClock, NAME_None, RF_Transactional);
+		NewMusicMap->SetFrameResolution(FocusedMovieScene->GetTickResolution());
+		NewMusicMap->Init(NewTempo, MusicalClock->MusicMap->GetInitialTimeSignature());
+
+		// This will remap from the old to the new map
+		UE::MovieScene::TimeHelpers::MigrateFrameTimes(FFrameBasedMusicMapRetimer(MusicalClock->MusicMap, NewMusicMap), FocusedMovieScene);
+		MusicalClock->MusicMap = NewMusicMap;
+	}
+}
+
+
+#undef LOCTEXT_NAMESPACE
\ No newline at end of file
diff --git a/Engine/Plugins/MovieScene/MusicalMode/Source/MusicalModeEditor/Private/SSequencerMusicModeCombo.h b/Engine/Plugins/MovieScene/MusicalMode/Source/MusicalModeEditor/Private/SSequencerMusicModeCombo.h
new file mode 100644
--- /dev/null
+++ b/Engine/Plugins/MovieScene/MusicalMode/Source/MusicalModeEditor/Private/SSequencerMusicModeCombo.h
@@ -0,0 +1,89 @@
+// Copyright Epic Games, Inc. All Rights Reserved.
+
+#pragma once
+
+#include "CoreTypes.h"
+#include "Templates/SharedPointer.h"
+#include "Containers/ArrayView.h"
+#include "Styling/SlateTypes.h"
+#include "Widgets/DeclarativeSyntaxSupport.h"
+#include "Widgets/SCompoundWidget.h"
+#include "Styling/AppStyle.h"
+#include "MovieSceneFwd.h"
+
+class ISequencer;
+class FMenuBuilder;
+class UFrameBasedMusicMap;
+class UMovieSceneMusicalClock;
+struct FMarkerProviderResults;
+
+class SSequencerMusicModeCombo : public SCompoundWidget
+{
+public:
+
+	SLATE_BEGIN_ARGS(SSequencerMusicModeCombo)
+		: _StyleSet(&FAppStyle::Get())
+		, _StyleName("SlimToolBar")
+	{}
+
+		SLATE_ARGUMENT(const ISlateStyle*, StyleSet)
+
+		SLATE_ARGUMENT(FName, StyleName)
+
+	SLATE_END_ARGS()
+
+	void Construct(const FArguments& InArgs, TWeakPtr<ISequencer> Sequencer);
+
+private:
+
+	bool GetIsSequenceReadOnly() const;
+
+	TSharedRef<SWidget> OnCreateMenu();
+
+	void PopulateSimpleMusicMapProperties(FMenuBuilder& MenuBuilder, const UFrameBasedMusicMap* MusicMap);
+	void PopulateSimpleMusicMapSection(FMenuBuilder& MenuBuilder, const UFrameBasedMusicMap* MusicMap);
+	void PopulateImportMusicMapSection(FMenuBuilder& MenuBuilder);
+	void PopulateMusicalClockSourceMenu(FMenuBuilder& MenuBuilder);
+
+	FText GetTempoTimeSignatureText() const;
+	FText GetToolTipText() const;
+	void OnRemoveAllMusicMaps();
+	void OnCreateSimpleMusicMap();
+	void OnImportMusicMap(UObject* Object);
+	void OnTimeSignatureTextCommitted(const FText& NewValue, ETextCommit::Type);
+	bool ValidateTimeSignatureString(const FText& InText, FText& OutErrorMessage) const;
+	bool SimpleMapEntryEnabled() const;
+	bool PromptForMarkerImport(const FMarkerProviderResults& Markers);
+
+	UMovieSceneMusicalClock* GetMusicalClock() const;
+
+	/** Sequencer pointer */
+	TWeakPtr<ISequencer> WeakSequencer;
+};
+
+
+class SSequencerMusicModeTempoEntryBox : public SCompoundWidget
+{
+public:
+
+	SLATE_BEGIN_ARGS(SSequencerMusicModeTempoEntryBox) : _UseSpinBoxStyle(false), _AuditionTempo(false)
+	{}
+		SLATE_ARGUMENT(bool, UseSpinBoxStyle)
+		SLATE_ARGUMENT(bool, AuditionTempo)
+	SLATE_END_ARGS()
+
+	void Construct(const FArguments& InArgs, TWeakPtr<ISequencer> InWeakSequencer);
+
+	virtual FReply OnMouseButtonDown(const FGeometry& MyGeometry, const FPointerEvent& MouseEvent) override;
+
+private:
+
+	void OnTempoTextCommitted(const FText& NewValue, ETextCommit::Type);
+	void SetValue(float NewTempo);
+	float GetValue() const;
+
+	/** Sequencer pointer */
+	TWeakPtr<ISequencer> WeakSequencer;
+
+	bool bAuditionTempo = false;
+};
\ No newline at end of file
