Fix NSLOCTEXT conflict

diff --git a/Engine/Plugins/MovieScene/MusicalMode/Source/MusicalModeEditor/Private/MusicalModeEditorModule.cpp b/Engine/Plugins/MovieScene/MusicalMode/Source/MusicalModeEditor/Private/MusicalModeEditorModule.cpp
--- a/Engine/Plugins/MovieScene/MusicalMode/Source/MusicalModeEditor/Private/MusicalModeEditorModule.cpp
+++ b/Engine/Plugins/MovieScene/MusicalMode/Source/MusicalModeEditor/Private/MusicalModeEditorModule.cpp
@@ -192,7 +192,7 @@ public:
 				FToolMenuEntry Entry = FToolMenuEntry::InitMenuEntry(
 					"BarseAndBeats",
 					NSLOCTEXT("MusicMode", "BarsBeatsText", "Bars : Beats"),
-					NSLOCTEXT("MusicMode", "BarsBeatsText", "Show time as Bars : Beats at the current tempo"),
+					NSLOCTEXT("MusicMode", "BarsBeatsText_ToolTip", "Show time as Bars : Beats at the current tempo"),
 					FSlateIcon(),
 					FToolUIActionChoice(
 						FUIAction(
